Fix abs value being rounded to integer.

Pointed out in PR #422.

diff --git a/Engine/Source/Developer/MeshSimplifier/Private/Quadric.cpp b/Engine/Source/Developer/MeshSimplifier/Private/Quadric.cpp
--- a/Engine/Source/Developer/MeshSimplifier/Private/Quadric.cpp
+++ b/Engine/Source/Developer/MeshSimplifier/Private/Quadric.cpp
@@ -29,7 +29,7 @@ bool CalcGradient( double grad[4], const FVector& p0, const FVector& p1, const F
 
 	det = - det3_201_123 * n[0] + det3_201_023 * n[1] - det3_201_013 * n[2];
 
-	if( abs( det ) < 1e-8 )
+	if( FMath::Abs( det ) < 1e-8 )
 	{
 		return false;
 	}
diff --git a/Engine/Source/Developer/MeshSimplifier/Private/Quadric.h b/Engine/Source/Developer/MeshSimplifier/Private/Quadric.h
--- a/Engine/Source/Developer/MeshSimplifier/Private/Quadric.h
+++ b/Engine/Source/Developer/MeshSimplifier/Private/Quadric.h
@@ -754,7 +754,7 @@ inline bool TQuadricAttrOptimizer< NumAttributes >::Optimize( FVector& Point ) c
 
 	double det = iNvx * nvx + iNvy * nvy + iNvz * nvz;
 
-	if( abs( det ) < 1e-8 )
+	if( FMath::Abs( det ) < 1e-8 )
 	{
 		return false;
 	}
