Tighten pragmas disabling shadow variables to make it clearer what needs to be fixed

diff --git a/Engine/Source/Developer/MeshSimplifier/Private/MeshSimplify.h b/Engine/Source/Developer/MeshSimplifier/Private/MeshSimplify.h
--- a/Engine/Source/Developer/MeshSimplifier/Private/MeshSimplify.h
+++ b/Engine/Source/Developer/MeshSimplifier/Private/MeshSimplify.h
@@ -1244,7 +1244,6 @@ void TMeshSimplifier<T, NumAttributes>::Collapse( TSimpEdge<T>* edge )
 	numVerts--;
 }
 
-PRAGMA_DISABLE_SHADOW_VARIABLE_WARNINGS
 template< typename T, uint32 NumAttributes >
 void TMeshSimplifier<T, NumAttributes>::SimplifyMesh( float maxError, int minTris )
 {
@@ -1400,7 +1399,9 @@ void TMeshSimplifier<T, NumAttributes>::SimplifyMesh( float maxError, int minTri
 			uint32	vertListNum = 0;
 			TSimpVert<T>*	vertList[256];
 			
+			PRAGMA_DISABLE_SHADOW_VARIABLE_WARNINGS
 			TSimpVert<T>* v = top->v1;
+			PRAGMA_POP
 			do {
 				vertList[ vertListNum++ ] = v;
 				v = v->next;
@@ -1409,7 +1410,9 @@ void TMeshSimplifier<T, NumAttributes>::SimplifyMesh( float maxError, int minTri
 			check( vertListNum <= 256 );
 			
 			for( uint32 i = 0; i < vertListNum; i++ ) {
+				PRAGMA_DISABLE_SHADOW_VARIABLE_WARNINGS
 				TSimpVert<T>* v = vertList[i];
+				PRAGMA_POP
 
 				if( v->TestFlags( SIMP_REMOVED ) ) {
 					// ungroup
@@ -1424,7 +1427,9 @@ void TMeshSimplifier<T, NumAttributes>::SimplifyMesh( float maxError, int minTri
 		{
 			// spread locked flag to vert group
 			uint32 flags = 0;
+			PRAGMA_DISABLE_SHADOW_VARIABLE_WARNINGS
 			TSimpVert<T>* v;
+			PRAGMA_POP
 
 			v = top->v1;
 			do {
@@ -1500,7 +1505,9 @@ void TMeshSimplifier<T, NumAttributes>::SimplifyMesh( float maxError, int minTri
 
 		// remove dead edges
 		for( int i = 0; i < edges.Num(); i++ ) {
+			PRAGMA_DISABLE_SHADOW_VARIABLE_WARNINGS
 			TSimpEdge<T>* edge = &edges[i];
+			PRAGMA_POP
 
 			if( edge->TestFlags( SIMP_REMOVED ) )
 				continue;
@@ -1516,7 +1523,9 @@ void TMeshSimplifier<T, NumAttributes>::SimplifyMesh( float maxError, int minTri
 			FHashTable HashTable( 4096, edges.Num() );
 
 			for( int i = 0; i < edges.Num(); i++ ) {
+				PRAGMA_DISABLE_SHADOW_VARIABLE_WARNINGS
 				TSimpEdge<T>* edge = &edges[i];
+				PRAGMA_POP
 
 				if( edge->TestFlags( SIMP_REMOVED ) )
 					continue;
@@ -1571,7 +1580,9 @@ void TMeshSimplifier<T, NumAttributes>::SimplifyMesh( float maxError, int minTri
 
 		// update edges
 		for( uint32 i = 0; i < updateEdgesNum; i++ ) {
+			PRAGMA_DISABLE_SHADOW_VARIABLE_WARNINGS
 			TSimpEdge<T>* edge = updateEdges[i];
+			PRAGMA_POP
 
 			edge->DisableFlags( SIMP_UPDATE );
 
@@ -1580,7 +1591,9 @@ void TMeshSimplifier<T, NumAttributes>::SimplifyMesh( float maxError, int minTri
 
 			float cost = ComputeEdgeCollapseCost( edge );
 
+			PRAGMA_DISABLE_SHADOW_VARIABLE_WARNINGS
 			TSimpEdge<T>* e = edge;
+			PRAGMA_POP
 			do {
 				uint32 EdgeIndex = GetEdgeIndex(e);
 				if( edgeHeap.IsPresent( EdgeIndex ) )
@@ -1635,7 +1648,6 @@ void TMeshSimplifier<T, NumAttributes>::SimplifyMesh( float maxError, int minTri
 
 	//common->Printf( "simplified numVerts %i, numTris %i\n", numVerts, numTris );
 }
-PRAGMA_POP
 
 template< typename T, uint32 NumAttributes >
 void TMeshSimplifier<T, NumAttributes>::OutputMesh( T* verts, uint32* indexes ) {
