Add a matching PRAGMA_ENABLE_SHADOW_VARIABLE_WARNINGS for PRAGMA_DISABLE_SHADOW_VARIABLE_WARNINGS instead of using basic PRAGMA_POP

diff --git a/Engine/Plugins/Developer/PerforceSourceControl/Source/PerforceSourceControl/Private/PerforceSourceControlPrivatePCH.h b/Engine/Plugins/Developer/PerforceSourceControl/Source/PerforceSourceControl/Private/PerforceSourceControlPrivatePCH.h
--- a/Engine/Plugins/Developer/PerforceSourceControl/Source/PerforceSourceControl/Private/PerforceSourceControlPrivatePCH.h
+++ b/Engine/Plugins/Developer/PerforceSourceControl/Source/PerforceSourceControl/Private/PerforceSourceControlPrivatePCH.h
@@ -9,6 +9,6 @@
 PRAGMA_DISABLE_SHADOW_VARIABLE_WARNINGS
 #include <p4/clientapi.h>
 #include <p4/i18napi.h>
-PRAGMA_POP
+PRAGMA_ENABLE_SHADOW_VARIABLE_WARNINGS
 
 #include "ISourceControlModule.h"
\ No newline at end of file
diff --git a/Engine/Source/Developer/ImageWrapper/Private/ExrImageWrapper.h b/Engine/Source/Developer/ImageWrapper/Private/ExrImageWrapper.h
--- a/Engine/Source/Developer/ImageWrapper/Private/ExrImageWrapper.h
+++ b/Engine/Source/Developer/ImageWrapper/Private/ExrImageWrapper.h
@@ -15,7 +15,7 @@ PRAGMA_DISABLE_SHADOW_VARIABLE_WARNINGS
 #include "ThirdParty/openexr/Deploy/include/ImfStdIO.h"
 #include "ThirdParty/openexr/Deploy/include/ImfChannelList.h"
 #include "ThirdParty/openexr/Deploy/include/ImfRgbaFile.h"
-PRAGMA_POP
+PRAGMA_ENABLE_SHADOW_VARIABLE_WARNINGS
 
 /**
  * OpenEXR implementation of the helper class
diff --git a/Engine/Source/Developer/MeshSimplifier/Private/MeshSimplify.h b/Engine/Source/Developer/MeshSimplifier/Private/MeshSimplify.h
--- a/Engine/Source/Developer/MeshSimplifier/Private/MeshSimplify.h
+++ b/Engine/Source/Developer/MeshSimplifier/Private/MeshSimplify.h
@@ -1401,7 +1401,7 @@ void TMeshSimplifier<T, NumAttributes>::SimplifyMesh( float maxError, int minTri
 			
 			PRAGMA_DISABLE_SHADOW_VARIABLE_WARNINGS
 			TSimpVert<T>* v = top->v1;
-			PRAGMA_POP
+			PRAGMA_ENABLE_SHADOW_VARIABLE_WARNINGS
 			do {
 				vertList[ vertListNum++ ] = v;
 				v = v->next;
@@ -1412,7 +1412,7 @@ void TMeshSimplifier<T, NumAttributes>::SimplifyMesh( float maxError, int minTri
 			for( uint32 i = 0; i < vertListNum; i++ ) {
 				PRAGMA_DISABLE_SHADOW_VARIABLE_WARNINGS
 				TSimpVert<T>* v = vertList[i];
-				PRAGMA_POP
+				PRAGMA_ENABLE_SHADOW_VARIABLE_WARNINGS
 
 				if( v->TestFlags( SIMP_REMOVED ) ) {
 					// ungroup
@@ -1429,7 +1429,7 @@ void TMeshSimplifier<T, NumAttributes>::SimplifyMesh( float maxError, int minTri
 			uint32 flags = 0;
 			PRAGMA_DISABLE_SHADOW_VARIABLE_WARNINGS
 			TSimpVert<T>* v;
-			PRAGMA_POP
+			PRAGMA_ENABLE_SHADOW_VARIABLE_WARNINGS
 
 			v = top->v1;
 			do {
@@ -1507,7 +1507,7 @@ void TMeshSimplifier<T, NumAttributes>::SimplifyMesh( float maxError, int minTri
 		for( int i = 0; i < edges.Num(); i++ ) {
 			PRAGMA_DISABLE_SHADOW_VARIABLE_WARNINGS
 			TSimpEdge<T>* edge = &edges[i];
-			PRAGMA_POP
+			PRAGMA_ENABLE_SHADOW_VARIABLE_WARNINGS
 
 			if( edge->TestFlags( SIMP_REMOVED ) )
 				continue;
@@ -1525,7 +1525,7 @@ void TMeshSimplifier<T, NumAttributes>::SimplifyMesh( float maxError, int minTri
 			for( int i = 0; i < edges.Num(); i++ ) {
 				PRAGMA_DISABLE_SHADOW_VARIABLE_WARNINGS
 				TSimpEdge<T>* edge = &edges[i];
-				PRAGMA_POP
+				PRAGMA_ENABLE_SHADOW_VARIABLE_WARNINGS
 
 				if( edge->TestFlags( SIMP_REMOVED ) )
 					continue;
@@ -1582,7 +1582,7 @@ void TMeshSimplifier<T, NumAttributes>::SimplifyMesh( float maxError, int minTri
 		for( uint32 i = 0; i < updateEdgesNum; i++ ) {
 			PRAGMA_DISABLE_SHADOW_VARIABLE_WARNINGS
 			TSimpEdge<T>* edge = updateEdges[i];
-			PRAGMA_POP
+			PRAGMA_ENABLE_SHADOW_VARIABLE_WARNINGS
 
 			edge->DisableFlags( SIMP_UPDATE );
 
@@ -1593,7 +1593,7 @@ void TMeshSimplifier<T, NumAttributes>::SimplifyMesh( float maxError, int minTri
 
 			PRAGMA_DISABLE_SHADOW_VARIABLE_WARNINGS
 			TSimpEdge<T>* e = edge;
-			PRAGMA_POP
+			PRAGMA_ENABLE_SHADOW_VARIABLE_WARNINGS
 			do {
 				uint32 EdgeIndex = GetEdgeIndex(e);
 				if( edgeHeap.IsPresent( EdgeIndex ) )
diff --git a/Engine/Source/Developer/MeshUtilities/Private/LayoutUV.cpp b/Engine/Source/Developer/MeshUtilities/Private/LayoutUV.cpp
--- a/Engine/Source/Developer/MeshUtilities/Private/LayoutUV.cpp
+++ b/Engine/Source/Developer/MeshUtilities/Private/LayoutUV.cpp
@@ -386,7 +386,7 @@ void FLayoutUV::FindCharts( const TMultiMap<int32,int32>& OverlappingCorners )
 
 		PRAGMA_DISABLE_SHADOW_VARIABLE_WARNINGS
 		uint32 NumTris = 0;
-		PRAGMA_POP
+		PRAGMA_ENABLE_SHADOW_VARIABLE_WARNINGS
 
 		for( int32 i = 0; i < Charts.Num(); i++ )
 		{
diff --git a/Engine/Source/Developer/ShaderFormatOpenGL/Private/GlslBackend.cpp b/Engine/Source/Developer/ShaderFormatOpenGL/Private/GlslBackend.cpp
--- a/Engine/Source/Developer/ShaderFormatOpenGL/Private/GlslBackend.cpp
+++ b/Engine/Source/Developer/ShaderFormatOpenGL/Private/GlslBackend.cpp
@@ -39,7 +39,7 @@
 
 PRAGMA_DISABLE_SHADOW_VARIABLE_WARNINGS
 #include "glsl_parser_extras.h"
-PRAGMA_POP
+PRAGMA_ENABLE_SHADOW_VARIABLE_WARNINGS
 
 #include "hash_table.h"
 #include "ir_rvalue_visitor.h"
diff --git a/Engine/Source/Developer/iOS/MetalShaderFormat/Private/MetalBackend.cpp b/Engine/Source/Developer/iOS/MetalShaderFormat/Private/MetalBackend.cpp
--- a/Engine/Source/Developer/iOS/MetalShaderFormat/Private/MetalBackend.cpp
+++ b/Engine/Source/Developer/iOS/MetalShaderFormat/Private/MetalBackend.cpp
@@ -10,7 +10,7 @@
 
 PRAGMA_DISABLE_SHADOW_VARIABLE_WARNINGS
 #include "glsl_parser_extras.h"
-PRAGMA_POP
+PRAGMA_ENABLE_SHADOW_VARIABLE_WARNINGS
 
 #include "hash_table.h"
 #include "ir_rvalue_visitor.h"
diff --git a/Engine/Source/Editor/UnrealEd/Private/EditorPhysXSupport.h b/Engine/Source/Editor/UnrealEd/Private/EditorPhysXSupport.h
--- a/Engine/Source/Editor/UnrealEd/Private/EditorPhysXSupport.h
+++ b/Engine/Source/Editor/UnrealEd/Private/EditorPhysXSupport.h
@@ -66,7 +66,7 @@ PRAGMA_DISABLE_SHADOW_VARIABLE_WARNINGS
 
 #pragma warning( pop )
 
-PRAGMA_POP
+PRAGMA_ENABLE_SHADOW_VARIABLE_WARNINGS
 
 using namespace physx;
 
diff --git a/Engine/Source/Runtime/Core/Private/Internationalization/ICUText.cpp b/Engine/Source/Runtime/Core/Private/Internationalization/ICUText.cpp
--- a/Engine/Source/Runtime/Core/Private/Internationalization/ICUText.cpp
+++ b/Engine/Source/Runtime/Core/Private/Internationalization/ICUText.cpp
@@ -12,7 +12,7 @@ PRAGMA_DISABLE_SHADOW_VARIABLE_WARNINGS
 #include <unicode/sortkey.h>
 #include <unicode/numfmt.h>
 #include <unicode/msgfmt.h>
-PRAGMA_POP
+PRAGMA_ENABLE_SHADOW_VARIABLE_WARNINGS
 
 #include "ICUUtilities.h"
 #include "ICUTextCharacterIterator.h"
diff --git a/Engine/Source/Runtime/Core/Public/Android/AndroidPlatformCompilerPreSetup.h b/Engine/Source/Runtime/Core/Public/Android/AndroidPlatformCompilerPreSetup.h
--- a/Engine/Source/Runtime/Core/Public/Android/AndroidPlatformCompilerPreSetup.h
+++ b/Engine/Source/Runtime/Core/Public/Android/AndroidPlatformCompilerPreSetup.h
@@ -47,6 +47,11 @@
 		_Pragma ("clang diagnostic ignored \"-Wshadow\"")
 #endif // PRAGMA_DISABLE_SHADOW_VARIABLE_WARNINGS
 
+#ifndef PRAGMA_ENABLE_SHADOW_VARIABLE_WARNINGS
+	#define PRAGMA_ENABLE_SHADOW_VARIABLE_WARNINGS \
+		_Pragma("clang diagnostic pop")
+#endif // PRAGMA_ENABLE_SHADOW_VARIABLE_WARNINGS
+
 #ifndef PRAGMA_POP
 	#define PRAGMA_POP \
 		_Pragma("clang diagnostic pop")
diff --git a/Engine/Source/Runtime/Core/Public/HTML5/HTML5PlatformCompilerPreSetup.h b/Engine/Source/Runtime/Core/Public/HTML5/HTML5PlatformCompilerPreSetup.h
--- a/Engine/Source/Runtime/Core/Public/HTML5/HTML5PlatformCompilerPreSetup.h
+++ b/Engine/Source/Runtime/Core/Public/HTML5/HTML5PlatformCompilerPreSetup.h
@@ -85,6 +85,11 @@
 		_Pragma ("clang diagnostic ignored \"-Wshadow\"")
 #endif // PRAGMA_DISABLE_SHADOW_VARIABLE_WARNINGS
 
+#ifndef PRAGMA_ENABLE_SHADOW_VARIABLE_WARNINGS
+	#define PRAGMA_ENABLE_SHADOW_VARIABLE_WARNINGS \
+		_Pragma("clang diagnostic pop")
+#endif // PRAGMA_ENABLE_SHADOW_VARIABLE_WARNINGS
+
 #ifndef PRAGMA_POP
 	#define PRAGMA_POP \
 		_Pragma("clang diagnostic pop")
diff --git a/Engine/Source/Runtime/Core/Public/IOS/IOSPlatformCompilerPreSetup.h b/Engine/Source/Runtime/Core/Public/IOS/IOSPlatformCompilerPreSetup.h
--- a/Engine/Source/Runtime/Core/Public/IOS/IOSPlatformCompilerPreSetup.h
+++ b/Engine/Source/Runtime/Core/Public/IOS/IOSPlatformCompilerPreSetup.h
@@ -83,6 +83,11 @@
 		_Pragma ("clang diagnostic ignored \"-Wshadow\"")
 #endif // PRAGMA_DISABLE_SHADOW_VARIABLE_WARNINGS
 
+#ifndef PRAGMA_ENABLE_SHADOW_VARIABLE_WARNINGS
+	#define PRAGMA_ENABLE_SHADOW_VARIABLE_WARNINGS \
+		_Pragma("clang diagnostic pop")
+#endif // PRAGMA_ENABLE_SHADOW_VARIABLE_WARNINGS
+
 #ifndef PRAGMA_POP
 	#define PRAGMA_POP \
 		_Pragma("clang diagnostic pop")
diff --git a/Engine/Source/Runtime/Core/Public/Linux/LinuxPlatformCompilerPreSetup.h b/Engine/Source/Runtime/Core/Public/Linux/LinuxPlatformCompilerPreSetup.h
--- a/Engine/Source/Runtime/Core/Public/Linux/LinuxPlatformCompilerPreSetup.h
+++ b/Engine/Source/Runtime/Core/Public/Linux/LinuxPlatformCompilerPreSetup.h
@@ -48,6 +48,11 @@
 		_Pragma ("clang diagnostic ignored \"-Wshadow\"")
 #endif // PRAGMA_DISABLE_SHADOW_VARIABLE_WARNINGS
 
+#ifndef PRAGMA_ENABLE_SHADOW_VARIABLE_WARNINGS
+	#define PRAGMA_ENABLE_SHADOW_VARIABLE_WARNINGS \
+		_Pragma("clang diagnostic pop")
+#endif // PRAGMA_ENABLE_SHADOW_VARIABLE_WARNINGS
+
 #ifndef PRAGMA_POP
 	#define PRAGMA_POP \
 		_Pragma("clang diagnostic pop")
diff --git a/Engine/Source/Runtime/Core/Public/Mac/MacPlatformCompilerPreSetup.h b/Engine/Source/Runtime/Core/Public/Mac/MacPlatformCompilerPreSetup.h
--- a/Engine/Source/Runtime/Core/Public/Mac/MacPlatformCompilerPreSetup.h
+++ b/Engine/Source/Runtime/Core/Public/Mac/MacPlatformCompilerPreSetup.h
@@ -83,6 +83,11 @@
 		_Pragma ("clang diagnostic ignored \"-Wshadow\"")
 #endif // PRAGMA_DISABLE_SHADOW_VARIABLE_WARNINGS
 
+#ifndef PRAGMA_ENABLE_SHADOW_VARIABLE_WARNINGS
+	#define PRAGMA_ENABLE_SHADOW_VARIABLE_WARNINGS \
+		_Pragma("clang diagnostic pop")
+#endif // PRAGMA_ENABLE_SHADOW_VARIABLE_WARNINGS
+
 #ifndef PRAGMA_POP
 	#define PRAGMA_POP \
 		_Pragma("clang diagnostic pop")
diff --git a/Engine/Source/Runtime/Core/Public/WinRT/WinRTPlatformCompilerPreSetup.h b/Engine/Source/Runtime/Core/Public/WinRT/WinRTPlatformCompilerPreSetup.h
--- a/Engine/Source/Runtime/Core/Public/WinRT/WinRTPlatformCompilerPreSetup.h
+++ b/Engine/Source/Runtime/Core/Public/WinRT/WinRTPlatformCompilerPreSetup.h
@@ -81,6 +81,11 @@
 		__pragma (warning(push))
 #endif // PRAGMA_DISABLE_SHADOW_VARIABLE_WARNINGS
 
+#ifndef PRAGMA_ENABLE_SHADOW_VARIABLE_WARNINGS
+	#define PRAGMA_ENABLE_SHADOW_VARIABLE_WARNINGS \
+		__pragma(warning(pop))
+#endif // PRAGMA_ENABLE_SHADOW_VARIABLE_WARNINGS
+
 #ifndef PRAGMA_POP
 	#define PRAGMA_POP \
 		__pragma(warning(pop))
diff --git a/Engine/Source/Runtime/Core/Public/Windows/WindowsPlatformCompilerPreSetup.h b/Engine/Source/Runtime/Core/Public/Windows/WindowsPlatformCompilerPreSetup.h
--- a/Engine/Source/Runtime/Core/Public/Windows/WindowsPlatformCompilerPreSetup.h
+++ b/Engine/Source/Runtime/Core/Public/Windows/WindowsPlatformCompilerPreSetup.h
@@ -60,9 +60,14 @@
 				_Pragma ("clang diagnostic ignored \"-Wshadow\"")
 	#endif // PRAGMA_DISABLE_SHADOW_VARIABLE_WARNINGS
 
+	#ifndef PRAGMA_ENABLE_SHADOW_VARIABLE_WARNINGS
+		#define PRAGMA_ENABLE_SHADOW_VARIABLE_WARNINGS \
+			_Pragma("clang diagnostic pop")
+	#endif // PRAGMA_ENABLE_SHADOW_VARIABLE_WARNINGS
+
 	#ifndef PRAGMA_POP
 		#define PRAGMA_POP \
-				_Pragma("clang diagnostic pop")
+			_Pragma("clang diagnostic pop")
 	#endif // PRAGMA_POP
 #else
 	// VC++
@@ -71,6 +76,11 @@
 			__pragma (warning(push))
 	#endif // PRAGMA_DISABLE_SHADOW_VARIABLE_WARNINGS
 
+	#ifndef PRAGMA_ENABLE_SHADOW_VARIABLE_WARNINGS
+		#define PRAGMA_ENABLE_SHADOW_VARIABLE_WARNINGS \
+			__pragma(warning(pop))
+	#endif // PRAGMA_ENABLE_SHADOW_VARIABLE_WARNINGS
+
 	#ifndef PRAGMA_POP
 		#define PRAGMA_POP \
 			__pragma(warning(pop))
diff --git a/Engine/Source/Runtime/Engine/Public/PhysXIncludes.h b/Engine/Source/Runtime/Engine/Public/PhysXIncludes.h
--- a/Engine/Source/Runtime/Engine/Public/PhysXIncludes.h
+++ b/Engine/Source/Runtime/Engine/Public/PhysXIncludes.h
@@ -72,7 +72,7 @@ PRAGMA_DISABLE_SHADOW_VARIABLE_WARNINGS
 
 MSVC_PRAGMA(warning(pop))
 
-PRAGMA_POP
+PRAGMA_ENABLE_SHADOW_VARIABLE_WARNINGS
 
 using namespace physx;
 
