Fixed NaN assert in Ocean

diff --git a/Engine/Source/Developer/MeshSimplifier/Private/QuadricMeshReduction.cpp b/Engine/Source/Developer/MeshSimplifier/Private/QuadricMeshReduction.cpp
--- a/Engine/Source/Developer/MeshSimplifier/Private/QuadricMeshReduction.cpp
+++ b/Engine/Source/Developer/MeshSimplifier/Private/QuadricMeshReduction.cpp
@@ -192,6 +192,11 @@ public:
 				NewVert.Tangents[1]		= InMesh.WedgeTangentY[ WedgeIndex ];
 				NewVert.Normal			= InMesh.WedgeTangentZ[ WedgeIndex ];
 
+				// Fix bad tangents
+				NewVert.Tangents[0] = NewVert.Tangents[0].ContainsNaN() ? FVector::ZeroVector : NewVert.Tangents[0];
+				NewVert.Tangents[1] = NewVert.Tangents[1].ContainsNaN() ? FVector::ZeroVector : NewVert.Tangents[1];
+				NewVert.Normal		= NewVert.Normal.ContainsNaN()		? FVector::ZeroVector : NewVert.Normal;
+
 				if( InMesh.WedgeColors.Num() == NumWedges )
 				{
 					NewVert.Color = FLinearColor::FromSRGBColor( InMesh.WedgeColors[ WedgeIndex ] );
