Fix render buffer optimisation
Fix UVs texture coordinnate
Fix vertex color import
Support LODs reduction

diff --git a/Engine/Source/Developer/MeshReductionInterface/Public/IMeshReductionInterfaces.h b/Engine/Source/Developer/MeshReductionInterface/Public/IMeshReductionInterfaces.h
--- a/Engine/Source/Developer/MeshReductionInterface/Public/IMeshReductionInterfaces.h
+++ b/Engine/Source/Developer/MeshReductionInterface/Public/IMeshReductionInterfaces.h
@@ -29,6 +29,22 @@ public:
 		const TMultiMap<int32, int32>& InOverlappingCorners,
 		const struct FMeshReductionSettings& ReductionSettings
 	) = 0;
+
+	/**
+	* Reduces the raw mesh using the provided reduction settings.
+	* @param OutReducedMesh - Upon return contains the reduced mesh.
+	* @param OutMaxDeviation - Upon return contains the maximum distance by which the reduced mesh deviates from the original.
+	* @param InMesh - The mesh to reduce.
+	* @param ReductionSettings - Setting with which to reduce the mesh.
+	*/
+	virtual void ReduceMeshDescription(
+		class UMeshDescription* OutReducedMesh,
+		float& OutMaxDeviation,
+		const class UMeshDescription* InMesh,
+		const TMultiMap<int32, int32>& InOverlappingCorners,
+		const struct FMeshReductionSettings& ReductionSettings
+	) = 0;
+
 	/**
 	* Reduces the provided skeletal mesh.
 	* @returns true if reduction was successful.
diff --git a/Engine/Source/Developer/MeshSimplifier/Private/QuadricMeshReduction.cpp b/Engine/Source/Developer/MeshSimplifier/Private/QuadricMeshReduction.cpp
--- a/Engine/Source/Developer/MeshSimplifier/Private/QuadricMeshReduction.cpp
+++ b/Engine/Source/Developer/MeshSimplifier/Private/QuadricMeshReduction.cpp
@@ -10,6 +10,7 @@
 #include "UniquePtr.h"
 #include "Features/IModularFeatures.h"
 #include "IMeshReductionInterfaces.h"
+#include "MeshDescription.h"
 
 class FQuadricSimplifierMeshReductionModule : public IMeshReductionModule
 {
@@ -436,6 +437,351 @@ public:
 		}
 	}
 
+	virtual void ReduceMeshDescription(
+		UMeshDescription* OutReducedMesh,
+		float& OutMaxDeviation,
+		const class UMeshDescription* InMesh,
+		const TMultiMap<int32, int32>& InOverlappingCorners,
+		const struct FMeshReductionSettings& ReductionSettings
+	) override
+	{
+		const uint32 NumTexCoords = MAX_STATIC_TEXCOORDS;
+		int32 InMeshNumTexCoords = 1;
+
+		TArray< TVertSimp< NumTexCoords > >	Verts;
+		TArray< uint32 >					Indexes;
+
+		TMap< int32, int32 > VertsMap;
+		TArray<int32> DupVerts;
+
+		int32 NumWedges = InMesh->VertexInstances().Num();
+		int32 NumFaces = NumWedges / 3;
+		check(InMesh->Polygons().Num() == NumFaces);
+
+		// Process each face, build vertex buffer and index buffer
+		for (int32 FaceIndex = 0; FaceIndex < NumFaces; FaceIndex++)
+		{
+			const FMeshVertexInstance* VertexInstances[3];
+			const FMeshVertex* Vertices[3];
+			FVector Positions[3];
+			for (int32 CornerIndex = 0; CornerIndex < 3; CornerIndex++)
+			{
+				VertexInstances[CornerIndex] = &(InMesh->GetVertexInstance(FVertexInstanceID((FaceIndex * 3) + CornerIndex)));
+				Vertices[CornerIndex] = &(InMesh->GetVertex(VertexInstances[CornerIndex]->VertexID));
+				Positions[CornerIndex] = Vertices[CornerIndex]->VertexPosition;
+			}
+
+			// Don't process degenerate triangles.
+			if (PointsEqual(Positions[0], Positions[1]) ||
+				PointsEqual(Positions[0], Positions[2]) ||
+				PointsEqual(Positions[1], Positions[2]))
+			{
+				continue;
+			}
+
+			int32 VertexIndices[3];
+			for (int32 CornerIndex = 0; CornerIndex < 3; CornerIndex++)
+			{
+				int32 WedgeIndex = FaceIndex * 3 + CornerIndex;
+
+				TVertSimp< NumTexCoords > NewVert;
+
+				if (VertexInstances[CornerIndex]->ConnectedPolygons.Num() > 0)
+				{
+					const FMeshPolygon& MeshPolygon = InMesh->GetPolygon(VertexInstances[CornerIndex]->ConnectedPolygons[0]);
+					NewVert.MaterialIndex = MeshPolygon.PolygonGroupID.GetValue();
+				}
+
+				NewVert.Position = Positions[CornerIndex];
+				NewVert.Tangents[0] = VertexInstances[CornerIndex]->Tangent;
+				NewVert.Normal = VertexInstances[CornerIndex]->Normal;
+				NewVert.Tangents[1] = FVector(0.0f);
+				if (!NewVert.Normal.IsNearlyZero(SMALL_NUMBER) && !NewVert.Tangents[0].IsNearlyZero(SMALL_NUMBER))
+				{
+					NewVert.Tangents[1] = FVector::CrossProduct(NewVert.Normal, NewVert.Tangents[0]).GetSafeNormal() * VertexInstances[CornerIndex]->BinormalSign;
+				}
+
+				// Fix bad tangents
+				NewVert.Tangents[0] = NewVert.Tangents[0].ContainsNaN() ? FVector::ZeroVector : NewVert.Tangents[0];
+				NewVert.Tangents[1] = NewVert.Tangents[1].ContainsNaN() ? FVector::ZeroVector : NewVert.Tangents[1];
+				NewVert.Normal = NewVert.Normal.ContainsNaN() ? FVector::ZeroVector : NewVert.Normal;
+				NewVert.Color = VertexInstances[CornerIndex]->Color;
+
+				for (int32 UVIndex = 0; UVIndex < NumTexCoords; UVIndex++)
+				{
+					if (VertexInstances[CornerIndex]->VertexUVs.IsValidIndex(UVIndex))
+					{
+						NewVert.TexCoords[UVIndex] = VertexInstances[CornerIndex]->VertexUVs[UVIndex];
+						InMeshNumTexCoords = FMath::Max(UVIndex, InMeshNumTexCoords);
+					}
+					else
+					{
+						NewVert.TexCoords[UVIndex] = FVector2D::ZeroVector;
+					}
+				}
+
+				// Make sure this vertex is valid from the start
+				NewVert.Correct();
+
+				DupVerts.Reset();
+				InOverlappingCorners.MultiFind(WedgeIndex, DupVerts);
+				DupVerts.Sort();
+
+				int32 Index = INDEX_NONE;
+				for (int32 k = 0; k < DupVerts.Num(); k++)
+				{
+					if (DupVerts[k] >= WedgeIndex)
+					{
+						// the verts beyond me haven't been placed yet, so these duplicates are not relevant
+						break;
+					}
+
+					int32* Location = VertsMap.Find(DupVerts[k]);
+					if (Location)
+					{
+						TVertSimp< NumTexCoords >& FoundVert = Verts[*Location];
+
+						if (NewVert.Equals(FoundVert))
+						{
+							Index = *Location;
+							break;
+						}
+					}
+				}
+				if (Index == INDEX_NONE)
+				{
+					Index = Verts.Add(NewVert);
+					VertsMap.Add(WedgeIndex, Index);
+				}
+				VertexIndices[CornerIndex] = Index;
+			}
+
+			// Reject degenerate triangles.
+			if (VertexIndices[0] == VertexIndices[1] ||
+				VertexIndices[1] == VertexIndices[2] ||
+				VertexIndices[0] == VertexIndices[2])
+			{
+				continue;
+			}
+
+			Indexes.Add(VertexIndices[0]);
+			Indexes.Add(VertexIndices[1]);
+			Indexes.Add(VertexIndices[2]);
+		}
+
+		uint32 NumVerts = Verts.Num();
+		uint32 NumIndexes = Indexes.Num();
+		uint32 NumTris = NumIndexes / 3;
+
+		static_assert(NumTexCoords == 8, "NumTexCoords changed, fix AttributeWeights");
+		const uint32 NumAttributes = (sizeof(TVertSimp< NumTexCoords >) - sizeof(uint32) - sizeof(FVector)) / sizeof(float);
+		float AttributeWeights[] =
+		{
+			16.0f, 16.0f, 16.0f,	// Normal
+			0.1f, 0.1f, 0.1f,		// Tangent[0]
+			0.1f, 0.1f, 0.1f,		// Tangent[1]
+			0.1f, 0.1f, 0.1f, 0.1f,	// Color
+			0.5f, 0.5f,				// TexCoord[0]
+			0.5f, 0.5f,				// TexCoord[1]
+			0.5f, 0.5f,				// TexCoord[2]
+			0.5f, 0.5f,				// TexCoord[3]
+			0.5f, 0.5f,				// TexCoord[4]
+			0.5f, 0.5f,				// TexCoord[5]
+			0.5f, 0.5f,				// TexCoord[6]
+			0.5f, 0.5f,				// TexCoord[7]
+		};
+		float* ColorWeights = AttributeWeights + 3 + 3 + 3;
+		float* TexCoordWeights = ColorWeights + 4;
+
+		// Zero out weights that aren't used
+		{
+			//TODO Check if we have vertex color
+
+			for (int32 TexCoordIndex = 0; TexCoordIndex < NumTexCoords; TexCoordIndex++)
+			{
+				if (InMesh->GetVertexInstance(FVertexInstanceID(0)).VertexUVs.Num() <= TexCoordIndex)
+				{
+					TexCoordWeights[2 * TexCoordIndex + 0] = 0.0f;
+					TexCoordWeights[2 * TexCoordIndex + 1] = 0.0f;
+				}
+			}
+		}
+
+		TMeshSimplifier< TVertSimp< NumTexCoords >, NumAttributes >* MeshSimp = new TMeshSimplifier< TVertSimp< NumTexCoords >, NumAttributes >(Verts.GetData(), NumVerts, Indexes.GetData(), NumIndexes);
+
+		MeshSimp->SetAttributeWeights(AttributeWeights);
+		//MeshSimp->SetBoundaryLocked();
+		MeshSimp->InitCosts();
+
+		float MaxErrorSqr = MeshSimp->SimplifyMesh(MAX_FLT, NumTris * ReductionSettings.PercentTriangles);
+
+		NumVerts = MeshSimp->GetNumVerts();
+		NumTris = MeshSimp->GetNumTris();
+		NumIndexes = NumTris * 3;
+
+		MeshSimp->OutputMesh(Verts.GetData(), Indexes.GetData());
+		delete MeshSimp;
+
+		OutMaxDeviation = FMath::Sqrt(MaxErrorSqr) / 8.0f;
+
+		{
+			//Empty the destination mesh
+			//We want to keep the polygon group information, the reduce do not impact the material list
+			for (const FPolygonGroupID& PolygonGroupID : OutReducedMesh->PolygonGroups().GetElementIDs())
+			{
+				FMeshPolygonGroup& PolygonGroup = OutReducedMesh->GetPolygonGroup(PolygonGroupID);
+				PolygonGroup.Polygons.Empty();
+			}
+			OutReducedMesh->Polygons().Reset();
+			OutReducedMesh->Edges().Reset();
+			OutReducedMesh->VertexInstances().Reset();
+			OutReducedMesh->Vertices().Reset();
+
+			//Fill the vertex array
+			for (int32 VertexIndex = 0; VertexIndex < (int32)NumVerts; ++VertexIndex)
+			{
+				FVertexID AddedVertexId = OutReducedMesh->CreateVertex();
+				FMeshVertex& NewVertex = OutReducedMesh->GetVertex(AddedVertexId);
+				NewVertex.VertexPosition = Verts[VertexIndex].Position;
+				NewVertex.CornerSharpness = 0.0f;
+				check(AddedVertexId.GetValue() == VertexIndex);
+			}
+
+			TMap<int32, FPolygonGroupID> PolygonGroupMapping;
+
+			//Vertex instances and Polygons
+			for (int32 TriangleIndex = 0; TriangleIndex < (int32)NumTris; TriangleIndex++)
+			{
+				FVertexInstanceID CornerInstanceIDs[3];
+				FVertexID CornerVerticesIDs[3];
+				for (int32 CornerIndex = 0; CornerIndex < 3; CornerIndex++)
+				{
+					int32 VertexInstanceIndex = TriangleIndex * 3 + CornerIndex;
+					const FVertexInstanceID VertexInstanceID(VertexInstanceIndex);
+					CornerInstanceIDs[CornerIndex] = VertexInstanceID;
+					int32 ControlPointIndex = Indexes[VertexInstanceIndex];
+					const FVertexID VertexID(ControlPointIndex);
+					//FVector VertexPosition = OutReducedMesh->GetVertex(VertexID).VertexPosition;
+					CornerVerticesIDs[CornerIndex] = VertexID;
+					FVertexInstanceID AddedVertexInstanceId = OutReducedMesh->CreateVertexInstance(VertexID);
+					//Make sure the Added vertex instance ID is matching the expected vertex instance ID
+					check(AddedVertexInstanceId == VertexInstanceID);
+					FMeshVertexInstance& NewVertexInstance = OutReducedMesh->GetVertexInstance(AddedVertexInstanceId);
+					check(AddedVertexInstanceId.GetValue() == VertexInstanceIndex);
+
+					//NTBs information
+					NewVertexInstance.Tangent = Verts[Indexes[VertexInstanceIndex]].Tangents[0];
+					NewVertexInstance.BinormalSign = GetBasisDeterminantSign(Verts[Indexes[VertexInstanceIndex]].Tangents[0].GetSafeNormal(), Verts[Indexes[VertexInstanceIndex]].Tangents[1].GetSafeNormal(), Verts[Indexes[VertexInstanceIndex]].Normal.GetSafeNormal());
+					NewVertexInstance.Normal = Verts[Indexes[VertexInstanceIndex]].Normal;
+
+					//Vertex Color
+					NewVertexInstance.Color = Verts[Indexes[VertexInstanceIndex]].Color;
+
+					//Texture coord
+					for (int32 TexCoordIndex = 0; TexCoordIndex < InMeshNumTexCoords; TexCoordIndex++)
+					{
+						while (NewVertexInstance.VertexUVs.Num() <= TexCoordIndex)
+						{
+							NewVertexInstance.VertexUVs.AddZeroed();
+						}
+						NewVertexInstance.VertexUVs[TexCoordIndex] = Verts[Indexes[VertexInstanceIndex]].TexCoords[TexCoordIndex];
+					}
+				}
+				
+				// material index
+				int32 MaterialIndex = Verts[Indexes[3 * TriangleIndex]].MaterialIndex;
+				FPolygonGroupID MaterialPolygonGroupID = FPolygonGroupID::Invalid;
+				if (!PolygonGroupMapping.Contains(MaterialIndex))
+				{
+					FPolygonGroupID PolygonGroupID(MaterialIndex);
+					check(InMesh->PolygonGroups().IsValid(PolygonGroupID));
+					const FMeshPolygonGroup& BasePolygonGroup = InMesh->GetPolygonGroup(PolygonGroupID);
+					MaterialPolygonGroupID = OutReducedMesh->PolygonGroups().Num() > MaterialIndex ? PolygonGroupID : OutReducedMesh->CreatePolygonGroup();
+					FMeshPolygonGroup& NewPolygonGroup = OutReducedMesh->GetPolygonGroup(MaterialPolygonGroupID);
+					NewPolygonGroup.bCastShadow = BasePolygonGroup.bCastShadow;
+					NewPolygonGroup.bEnableCollision = BasePolygonGroup.bEnableCollision;
+					NewPolygonGroup.ImportedMaterialSlotName = BasePolygonGroup.ImportedMaterialSlotName;
+					NewPolygonGroup.MaterialSlotName = BasePolygonGroup.MaterialSlotName;
+					NewPolygonGroup.MaterialAsset = BasePolygonGroup.MaterialAsset;
+					PolygonGroupMapping.Add(MaterialIndex, MaterialPolygonGroupID);
+				}
+				else
+				{
+					MaterialPolygonGroupID = PolygonGroupMapping[MaterialIndex];
+				}
+
+				// Create polygon edges
+				TArray<UMeshDescription::FContourPoint> Contours;
+				{
+					// Add the edges of this triangle
+					for (uint32 TriangleEdgeNumber = 0; TriangleEdgeNumber < 3; ++TriangleEdgeNumber)
+					{
+						int32 ContourPointIndex = Contours.AddDefaulted();
+						UMeshDescription::FContourPoint& ContourPoint = Contours[ContourPointIndex];
+						//Find the matching edge ID
+						uint32 CornerIndices[2];
+						CornerIndices[0] = (TriangleEdgeNumber + 0) % 3;
+						CornerIndices[1] = (TriangleEdgeNumber + 1) % 3;
+
+						FVertexID EdgeVertexIDs[2];
+						EdgeVertexIDs[0] = CornerVerticesIDs[CornerIndices[0]];
+						EdgeVertexIDs[1] = CornerVerticesIDs[CornerIndices[1]];
+
+						FEdgeID MatchEdgeId = FEdgeID::Invalid;
+						const FMeshVertex& EdgeMeshVertex = OutReducedMesh->GetVertex(EdgeVertexIDs[0]);
+						for (const FEdgeID& EdgeID : EdgeMeshVertex.ConnectedEdgeIDs)
+						{
+							FMeshEdge& Edge = OutReducedMesh->GetEdge(EdgeID);
+							if (EdgeVertexIDs[0] == Edge.VertexIDs[0] && EdgeVertexIDs[1] == Edge.VertexIDs[1] ||
+								EdgeVertexIDs[1] == Edge.VertexIDs[0] && EdgeVertexIDs[0] == Edge.VertexIDs[1])
+							{
+								MatchEdgeId = EdgeID;
+								break;
+							}
+						}
+						if (MatchEdgeId == FEdgeID::Invalid)
+						{
+							MatchEdgeId = OutReducedMesh->CreateEdge(EdgeVertexIDs[0], EdgeVertexIDs[1]);
+						}
+						ContourPoint.EdgeID = MatchEdgeId;
+						ContourPoint.VertexInstanceID = CornerInstanceIDs[CornerIndices[0]];
+					}
+				}
+
+				// Insert a polygon into the mesh
+				const FPolygonID NewPolygonID = OutReducedMesh->CreatePolygon(Contours);
+				FMeshPolygonGroup& MeshPolygonGroup = OutReducedMesh->GetPolygonGroup(MaterialPolygonGroupID);
+				FMeshPolygon& NewPolygon = OutReducedMesh->GetPolygon(NewPolygonID);
+				NewPolygon.PolygonGroupID = MaterialPolygonGroupID;
+				MeshPolygonGroup.Polygons.Add(NewPolygonID);
+				int32 NewTriangleIndex = NewPolygon.Triangles.AddDefaulted();
+				FMeshTriangle& NewTriangle = NewPolygon.Triangles[NewTriangleIndex];
+				for (int32 TriangleVertexIndex = 0; TriangleVertexIndex < 3; ++TriangleVertexIndex)
+				{
+					const FVertexInstanceID &VertexInstanceID = CornerInstanceIDs[TriangleVertexIndex];
+					NewTriangle.SetVertexInstanceID(TriangleVertexIndex, VertexInstanceID);
+				}
+			}
+			Verts.Empty();
+			Indexes.Empty();
+
+			//Remove the unused polygon group (reduce can remove all polygons from a group)
+			TArray<FPolygonGroupID> ToDeletePolygonGroupIDs;
+			for (const FPolygonGroupID& PolygonGroupID : OutReducedMesh->PolygonGroups().GetElementIDs())
+			{
+				FMeshPolygonGroup& PolygonGroup = OutReducedMesh->GetPolygonGroup(PolygonGroupID);
+				if (PolygonGroup.Polygons.Num() == 0)
+				{
+					ToDeletePolygonGroupIDs.Add(PolygonGroupID);
+				}
+			}
+			for (const FPolygonGroupID& PolygonGroupID : ToDeletePolygonGroupIDs)
+			{
+				OutReducedMesh->DeletePolygonGroup(PolygonGroupID);
+			}
+		}
+	}
+
 	virtual bool ReduceSkeletalMesh(
 		USkeletalMesh* SkeletalMesh,
 		int32 LODIndex,
diff --git a/Engine/Source/Developer/MeshSimplifier/QuadricMeshReduction.Build.cs b/Engine/Source/Developer/MeshSimplifier/QuadricMeshReduction.Build.cs
--- a/Engine/Source/Developer/MeshSimplifier/QuadricMeshReduction.Build.cs
+++ b/Engine/Source/Developer/MeshSimplifier/QuadricMeshReduction.Build.cs
@@ -13,6 +13,7 @@ public class QuadricMeshReduction : ModuleRules
 		PrivateDependencyModuleNames.Add("Engine");
 		PrivateDependencyModuleNames.Add("RenderCore");
 		PrivateDependencyModuleNames.Add("RawMesh");
+        PrivateDependencyModuleNames.Add("MeshDescription");
 
         PrivateIncludePathModuleNames.AddRange(
         new string[] {
diff --git a/Engine/Source/Developer/SimplygonMeshReduction/Private/SimplygonMeshReduction.cpp b/Engine/Source/Developer/SimplygonMeshReduction/Private/SimplygonMeshReduction.cpp
--- a/Engine/Source/Developer/SimplygonMeshReduction/Private/SimplygonMeshReduction.cpp
+++ b/Engine/Source/Developer/SimplygonMeshReduction/Private/SimplygonMeshReduction.cpp
@@ -493,6 +493,16 @@ public:
 		SkeletalMesh->OptimizationSettings[LODIndex] = Settings;
 	}
 
+	virtual void ReduceMeshDescription(
+		class UMeshDescription* OutReducedMesh,
+		float& OutMaxDeviation,
+		const class UMeshDescription* InMesh,
+		const TMultiMap<int32, int32>& InOverlappingCorners,
+		const struct FMeshReductionSettings& ReductionSettings
+	) override
+	{
+	}
+
 	virtual bool ReduceSkeletalMesh(
 		USkeletalMesh* SkeletalMesh,
 		int32 LODIndex,
diff --git a/Engine/Source/Editor/UnrealEd/Private/Fbx/FbxStaticMeshImport.cpp b/Engine/Source/Editor/UnrealEd/Private/Fbx/FbxStaticMeshImport.cpp
--- a/Engine/Source/Editor/UnrealEd/Private/Fbx/FbxStaticMeshImport.cpp
+++ b/Engine/Source/Editor/UnrealEd/Private/Fbx/FbxStaticMeshImport.cpp
@@ -526,10 +526,28 @@ bool UnFbx::FFbxImporter::BuildStaticMeshFromGeometry(FbxNode* Node, UStaticMesh
 
 		int32 MaxMaterialIndex = 0;
 		TMap<int32, FPolygonGroupID> PolygonGroupMapping;
-		int32 NumUVs = FBXUVs.UniqueUVCount;
+
+		// When importing multiple mesh pieces to the same static mesh.  Ensure each mesh piece has the same number of Uv's
+		int32 ExistingUVCount = 0;
+		for (const FVertexInstanceID& VertexInstanceID : MeshDescription->VertexInstances().GetElementIDs())
+		{
+			ExistingUVCount = FMath::Max(ExistingUVCount, MeshDescription->GetVertexInstance(VertexInstanceID).VertexUVs.Num());
+		}
+		
+		int32 NumUVs = FMath::Max(FBXUVs.UniqueUVCount, ExistingUVCount);
 		// At least one UV set must exist.  
 		NumUVs = FMath::Max(1, NumUVs);
 
+		//Make sure all Vertex instance have the correct number of UVs
+		for (const FVertexInstanceID& VertexInstanceID : MeshDescription->VertexInstances().GetElementIDs())
+		{
+			FMeshVertexInstance& VertexInstance = MeshDescription->GetVertexInstance(VertexInstanceID);
+			while (VertexInstance.VertexUVs.Num() < NumUVs)
+			{
+				VertexInstance.VertexUVs.Add(FVector2D(0.0f, 0.0f));
+			}
+		}
+
 		//Fill the vertex array
 		for (int32 VertexIndex = 0; VertexIndex < VertexCount; ++VertexIndex)
 		{
@@ -1834,67 +1852,6 @@ UStaticMesh* UnFbx::FFbxImporter::ImportStaticMeshAsSingle(UObject* InParent, TA
 		}
 		else
 		{
-			// Compress the materials array by removing any duplicates.
-			bool bDoRemap = false;
-			TArray<int32> MaterialMap;
-			TArray<FFbxMaterial> UniqueMaterials;
-			for (int32 MaterialIndex = 0; MaterialIndex < MeshMaterials.Num(); ++MaterialIndex)
-			{
-				bool bUnique = true;
-				for (int32 OtherMaterialIndex = MaterialIndex - 1; OtherMaterialIndex >= 0; --OtherMaterialIndex)
-				{
-					if (MeshMaterials[MaterialIndex].FbxMaterial == MeshMaterials[OtherMaterialIndex].FbxMaterial &&
-						MeshMaterials[MaterialIndex].Material == MeshMaterials[OtherMaterialIndex].Material)
-					{
-						int32 UniqueIndex = MaterialMap[OtherMaterialIndex];
-						MaterialMap.Add(UniqueIndex);
-						bDoRemap = true;
-						bUnique = false;
-						break;
-					}
-				}
-				if (bUnique)
-				{
-					int32 UniqueIndex = UniqueMaterials.Add(MeshMaterials[MaterialIndex]);
-					MaterialMap.Add(UniqueIndex);
-				}
-				else
-				{
-					UE_LOG(LogFbx, Verbose, TEXT("  remap %d -> %d"), MaterialIndex, MaterialMap[MaterialIndex]);
-				}
-			}
-
-			if (UniqueMaterials.Num() > LARGE_MESH_MATERIAL_INDEX_THRESHOLD)
-			{
-				AddTokenizedErrorMessage(
-					FTokenizedMessage::Create(
-						EMessageSeverity::Warning,
-						FText::Format(LOCTEXT("Error_TooManyMaterials", "StaticMesh has a large number({0}) of materials and may render inefficently.  Consider breaking up the mesh into multiple Static Mesh Assets"),
-							FText::AsNumber(UniqueMaterials.Num())
-						)),
-					FFbxErrors::StaticMesh_TooManyMaterials);
-			}
-
-			if (bDoRemap)
-			{
-				int32 GroupNumber = MeshDescription->PolygonGroups().Num();
-
-				for (const FPolygonID& PolygonID : MeshDescription->Polygons().GetElementIDs())
-				{
-					FMeshPolygon MeshPolygon = MeshDescription->GetPolygon(PolygonID);
-					if (MaterialMap[MeshPolygon.PolygonGroupID.GetValue()] != MeshPolygon.PolygonGroupID.GetValue())
-					{
-						//Remove from the old polygon group
-						FMeshPolygonGroup& SourcePolygonGroup = MeshDescription->GetPolygonGroup(MeshPolygon.PolygonGroupID);
-						SourcePolygonGroup.Polygons.Remove(PolygonID);
-						//Set the new ID
-						MeshPolygon.PolygonGroupID = FPolygonGroupID(MaterialMap[MeshPolygon.PolygonGroupID.GetValue()]);
-						//Add to the new polygon group
-						FMeshPolygonGroup& DestinationPolygonGroup = MeshDescription->GetPolygonGroup(MeshPolygon.PolygonGroupID);
-						DestinationPolygonGroup.Polygons.Add(PolygonID);
-					}
-				}
-			}
 			TArray<FStaticMaterial> MaterialToAdd;
 			for (const FPolygonGroupID& PolygonGroupID : MeshDescription->PolygonGroups().GetElementIDs())
 			{
diff --git a/Engine/Source/Runtime/Engine/Private/StaticMesh.cpp b/Engine/Source/Runtime/Engine/Private/StaticMesh.cpp
--- a/Engine/Source/Runtime/Engine/Private/StaticMesh.cpp
+++ b/Engine/Source/Runtime/Engine/Private/StaticMesh.cpp
@@ -1296,7 +1296,7 @@ void FStaticMeshRenderData::Cache(UStaticMesh* Owner, const FStaticMeshLODSettin
 			if (Owner->GetOriginalMeshDescription(0) != nullptr)
 			{
 				IMeshBuilderModule& MeshBuilderModule = FModuleManager::Get().LoadModuleChecked<IMeshBuilderModule>(TEXT("MeshBuilder"));
-				MeshBuilderModule.BuildMesh(Owner);
+				MeshBuilderModule.BuildMesh(Owner, LODGroup);
 			}
 			else
 			{
diff --git a/Engine/Source/Runtime/MeshBuilder/Private/BuildOptimizationHelper.h b/Engine/Source/Runtime/MeshBuilder/Private/BuildOptimizationHelper.h
--- a/Engine/Source/Runtime/MeshBuilder/Private/BuildOptimizationHelper.h
+++ b/Engine/Source/Runtime/MeshBuilder/Private/BuildOptimizationHelper.h
@@ -312,6 +312,21 @@ public:
 		}
 	}
 
+	static void CacheOptimizeIndexBuffer(TArray<uint16>& Indices)
+	{
+		bool bDisableTriangleOrderOptimization = (BuildOptimizationPrivate::CVarTriangleOrderOptimization.GetValueOnGameThread() == 2);
+		bool bUsingNvTriStrip = !bDisableTriangleOrderOptimization && (BuildOptimizationPrivate::CVarTriangleOrderOptimization.GetValueOnGameThread() == 0);
+
+		if (bUsingNvTriStrip)
+		{
+			BuildOptimizationPrivate::NvTriStripHelper::CacheOptimizeIndexBuffer(Indices);
+		}
+		else if (!bDisableTriangleOrderOptimization)
+		{
+			BuildOptimizationPrivate::ForsythHelper::CacheOptimizeIndexBuffer(Indices);
+		}
+	}
+
 	static void CacheOptimizeIndexBuffer(TArray<uint32>& Indices)
 	{
 		bool bDisableTriangleOrderOptimization = (BuildOptimizationPrivate::CVarTriangleOrderOptimization.GetValueOnGameThread() == 2);
@@ -333,14 +348,25 @@ public:
 		TArray<int32>& WedgeMap
 	)
 	{
+		//////////////////////////////////////////////////////////////////////////
+
+		// TODO
+		// this code is commented because it make the indexes incoherent
+		// Is it something we must do? Not sure!
+		// We should investigate why.
+		// Currently we at least optimize the section indexes, but we do not reorder the vertex array
+
+		//////////////////////////////////////////////////////////////////////////
+
+
 		// Copy the vertices since we will be reordering them
-		TArray<FStaticMeshBuildVertex> OriginalVertices = Vertices;
+		//TArray<FStaticMeshBuildVertex> OriginalVertices = Vertices;
 
 		// Initialize a cache that stores which indices have been assigned
-		TArray<int32> IndexCache;
-		IndexCache.AddUninitialized(Vertices.Num());
-		FMemory::Memset(IndexCache.GetData(), INDEX_NONE, IndexCache.Num() * IndexCache.GetTypeSize());
-		int32 NextAvailableIndex = 0;
+		//TArray<int32> IndexCache;
+		//IndexCache.AddUninitialized(Vertices.Num());
+		//FMemory::Memset(IndexCache.GetData(), INDEX_NONE, IndexCache.Num() * IndexCache.GetTypeSize());
+		//int32 NextAvailableIndex = 0;
 
 		// Iterate through the section index buffers, 
 		// Optimizing index order for the post transform cache (minimizes the number of vertices transformed), 
@@ -353,7 +379,7 @@ public:
 			{
 				// Optimize the index buffer for the post transform cache with.
 				CacheOptimizeIndexBuffer(Indices);
-
+/*
 				// Copy the index buffer since we will be reordering it
 				TArray<uint32> OriginalIndices = Indices;
 
@@ -377,10 +403,12 @@ public:
 					}
 					// Reorder the vertices based on the new index assignment
 					Vertices[Indices[Index]] = OriginalVertices[OriginalIndices[Index]];
+
 				}
+*/
 			}
 		}
-
+/*
 		for (int32 i = 0; i < WedgeMap.Num(); i++)
 		{
 			int32 MappedIndex = WedgeMap[i];
@@ -389,5 +417,6 @@ public:
 				WedgeMap[i] = IndexCache[MappedIndex];
 			}
 		}
+*/
 	}
 };
\ No newline at end of file
diff --git a/Engine/Source/Runtime/MeshBuilder/Private/MeshBuilderModule.cpp b/Engine/Source/Runtime/MeshBuilder/Private/MeshBuilderModule.cpp
--- a/Engine/Source/Runtime/MeshBuilder/Private/MeshBuilderModule.cpp
+++ b/Engine/Source/Runtime/MeshBuilder/Private/MeshBuilderModule.cpp
@@ -24,7 +24,7 @@ public:
 		// Unregister any modular features here
 	}
 
-	virtual bool BuildMesh(class UObject* Mesh) override;
+	virtual bool BuildMesh(UObject* Mesh, const FStaticMeshLODGroup& LODGroup) override;
 
 private:
 
@@ -32,13 +32,13 @@ private:
 
 IMPLEMENT_MODULE(FMeshBuilderModule, MeshBuilder );
 
-bool FMeshBuilderModule::BuildMesh(class UObject* Mesh)
+bool FMeshBuilderModule::BuildMesh(class UObject* Mesh, const FStaticMeshLODGroup& LODGroup)
 {
 	UStaticMesh* StaticMesh = Cast<UStaticMesh>(Mesh);
 	if (StaticMesh != nullptr)
 	{
 		//Call the static mesh builder
-		return FStaticMeshBuilder().Build(StaticMesh);
+		return FStaticMeshBuilder().Build(StaticMesh, LODGroup);
 	}
 	return false;
 }
diff --git a/Engine/Source/Runtime/MeshBuilder/Private/MeshDescriptionHelper.cpp b/Engine/Source/Runtime/MeshBuilder/Private/MeshDescriptionHelper.cpp
--- a/Engine/Source/Runtime/MeshBuilder/Private/MeshDescriptionHelper.cpp
+++ b/Engine/Source/Runtime/MeshBuilder/Private/MeshDescriptionHelper.cpp
@@ -13,6 +13,9 @@
 #include "Serialization/MemoryReader.h"
 #include "RenderUtils.h"
 #include "LayoutUV.h"
+#include "IMeshReductionInterfaces.h"
+
+#include "IMeshReductionManagerModule.h"
 
 DEFINE_LOG_CATEGORY(LogMeshBuilder);
 
@@ -86,7 +89,7 @@ UMeshDescription* FMeshDescriptionHelper::GetRenderMeshDescription(UObject* Owne
 		CreatePolygonNTB(RenderMeshDescription, BuildSettings->bRemoveDegenerates ? SMALL_NUMBER : 0.0f);
 		
 		// Find overlapping corners to accelerate adjacency.
-		FindOverlappingCorners(ComparisonThreshold);
+		FindOverlappingCorners(OverlappingCorners, RenderMeshDescription, ComparisonThreshold);
 			
 		//Keep the original mesh description NTBs if we do not rebuild the normals or tangents.
 		bool bComputeTangentLegacy = !BuildSettings->bUseMikkTSpace && (BuildSettings->bRecomputeNormals || BuildSettings->bRecomputeTangents);
@@ -126,6 +129,26 @@ UMeshDescription* FMeshDescriptionHelper::GetRenderMeshDescription(UObject* Owne
 	return RenderMeshDescription;
 }
 
+void FMeshDescriptionHelper::ReduceLOD(const UMeshDescription* BaseMesh, UMeshDescription* DestMesh, const FMeshReductionSettings& ReductionSettings, const TMultiMap<int32, int32>& InOverlappingCorners)
+{
+	if (BaseMesh == nullptr || DestMesh == nullptr)
+	{
+		//UE_LOG(LogMeshUtilities, Error, TEXT("Mesh contains zero source models."));
+		return;
+	}
+
+	IMeshReductionManagerModule& MeshReductionModule = FModuleManager::Get().LoadModuleChecked<IMeshReductionManagerModule>("MeshReductionInterface");
+	IMeshReduction* MeshReduction = MeshReductionModule.GetStaticMeshReductionInterface();
+	// Reduce this LOD mesh according to its reduction settings.
+
+	if (!MeshReduction || (ReductionSettings.PercentTriangles >= 1.0f && ReductionSettings.MaxDeviation <= 0.0f))
+	{
+		return;
+	}
+	float MaxDeviation = ReductionSettings.MaxDeviation;
+	MeshReduction->ReduceMeshDescription(DestMesh, MaxDeviation, BaseMesh, InOverlappingCorners, ReductionSettings);
+}
+
 void FMeshDescriptionHelper::CopyMeshDescription(UMeshDescription* SourceMeshDescription, UMeshDescription* DestinationMeshDescription) const
 {
 	//Copy the Source into the destination
@@ -143,28 +166,28 @@ bool FMeshDescriptionHelper::IsValidOriginalMeshDescription()
 	return OriginalMeshDescription != nullptr;
 }
 
-void FMeshDescriptionHelper::FindOverlappingCorners(float ComparisonThreshold)
+void FMeshDescriptionHelper::FindOverlappingCorners(TMultiMap<int32, int32>& OverlappingCorners, const UMeshDescription* MeshDescription, float ComparisonThreshold)
 {
 	//Empty the old data
 	OverlappingCorners.Reset();
-	const int32 NumWedges = OriginalMeshDescription->VertexInstances().Num();
+	const int32 NumWedges = MeshDescription->VertexInstances().Num();
 
 	// Create a list of vertex Z/index pairs
 	TArray<FIndexAndZ> VertIndexAndZ;
 	VertIndexAndZ.Reserve(NumWedges);
-	const FVertexInstanceArray& VertexInstances = OriginalMeshDescription->VertexInstances();
+	const FVertexInstanceArray& VertexInstances = MeshDescription->VertexInstances();
 
-	for (const FPolygonID& PolygonID : OriginalMeshDescription->Polygons().GetElementIDs())
+	for (const FPolygonID& PolygonID : MeshDescription->Polygons().GetElementIDs())
 	{
-		const TArray<FMeshTriangle>& MeshTriangles = OriginalMeshDescription->GetPolygon(PolygonID).Triangles;
+		const TArray<FMeshTriangle>& MeshTriangles = MeshDescription->GetPolygon(PolygonID).Triangles;
 		for (const FMeshTriangle& MeshTriangle : MeshTriangles)
 		{
 			for (int32 CornerIndex = 0; CornerIndex < 3; ++CornerIndex)
 			{
 				const FVertexInstanceID& VertexInstanceID = MeshTriangle.GetVertexInstanceID(CornerIndex);
-				const FMeshVertexInstance& VertexInstance = OriginalMeshDescription->GetVertexInstance(VertexInstanceID);
+				const FMeshVertexInstance& VertexInstance = MeshDescription->GetVertexInstance(VertexInstanceID);
 
-				new(VertIndexAndZ)FIndexAndZ(VertexInstanceID.GetValue(), OriginalMeshDescription->GetVertex(VertexInstance.VertexID).VertexPosition);
+				new(VertIndexAndZ)FIndexAndZ(VertexInstanceID.GetValue(), MeshDescription->GetVertex(VertexInstance.VertexID).VertexPosition);
 			}
 		}
 	}
@@ -181,8 +204,8 @@ void FMeshDescriptionHelper::FindOverlappingCorners(float ComparisonThreshold)
 			if (FMath::Abs(VertIndexAndZ[j].Z - VertIndexAndZ[i].Z) > ComparisonThreshold)
 				break; // can't be any more dups
 
-			const FVector& PositionA = OriginalMeshDescription->GetVertex(OriginalMeshDescription->GetVertexInstance(FVertexInstanceID(VertIndexAndZ[i].Index)).VertexID).VertexPosition;
-			const FVector& PositionB = OriginalMeshDescription->GetVertex(OriginalMeshDescription->GetVertexInstance(FVertexInstanceID(VertIndexAndZ[j].Index)).VertexID).VertexPosition;
+			const FVector& PositionA = MeshDescription->GetVertex(MeshDescription->GetVertexInstance(FVertexInstanceID(VertIndexAndZ[i].Index)).VertexID).VertexPosition;
+			const FVector& PositionB = MeshDescription->GetVertex(MeshDescription->GetVertexInstance(FVertexInstanceID(VertIndexAndZ[j].Index)).VertexID).VertexPosition;
 
 			if (PositionA.Equals(PositionB, ComparisonThreshold))
 			{
@@ -193,6 +216,11 @@ void FMeshDescriptionHelper::FindOverlappingCorners(float ComparisonThreshold)
 	}
 }
 
+void FMeshDescriptionHelper::FindOverlappingCorners(const UMeshDescription* MeshDescription, float ComparisonThreshold)
+{
+	FindOverlappingCorners(OverlappingCorners, MeshDescription, ComparisonThreshold);
+}
+
 const FVector& FMeshDescriptionHelper::GetVertexPositionFromVertexInstance(UMeshDescription* MeshDescription, const FVertexInstanceID& VertexInstanceID) const
 {
 	return MeshDescription->GetVertex(MeshDescription->GetVertexInstance(VertexInstanceID).VertexID).VertexPosition;
diff --git a/Engine/Source/Runtime/MeshBuilder/Private/MeshDescriptionHelper.h b/Engine/Source/Runtime/MeshBuilder/Private/MeshDescriptionHelper.h
--- a/Engine/Source/Runtime/MeshBuilder/Private/MeshDescriptionHelper.h
+++ b/Engine/Source/Runtime/MeshBuilder/Private/MeshDescriptionHelper.h
@@ -41,11 +41,16 @@ public:
 	//Build a render mesh description with the BuildSettings. This will update the InRenderMeshDescription ptr content
 	UMeshDescription* GetRenderMeshDescription(UObject* Owner);
 
+	void ReduceLOD(const UMeshDescription* BaseMesh, UMeshDescription* DestMesh, const struct FMeshReductionSettings& ReductionSettings, const TMultiMap<int32, int32>& InOverlappingCorners);
+
 	void CopyMeshDescription(UMeshDescription* SourceMeshDescription, UMeshDescription* DestinationMeshDescription) const;
 
 	//Return true if there is a valid original mesh description, false otherwise(Auto generate LOD).
 	bool IsValidOriginalMeshDescription();
 
+	static void FindOverlappingCorners(TMultiMap<int32, int32>& OverlappingCorners, const UMeshDescription* MeshDescription, float ComparisonThreshold);
+	void FindOverlappingCorners(const UMeshDescription* MeshDescription, float ComparisonThreshold);
+
 	const TMultiMap<int32, int32>& GetOverlappingCorners() const { return OverlappingCorners; }
 
 private:
@@ -58,7 +63,6 @@ private:
 	void CreateNormals(UMeshDescription* MeshDescription, ETangentOptions TangentOptions, bool bComputeTangent);
 	void CreateMikktTangents(UMeshDescription* MeshDescription, ETangentOptions TangentOptions);
 	void CreatePolygonNTB(UMeshDescription* MeshDescription, float ComparisonThreshold);
-	void FindOverlappingCorners(float ComparisonThreshold);
 
 	//////////////////////////////////////////////////////////////////////////
 	//PRIVATE class members
diff --git a/Engine/Source/Runtime/MeshBuilder/Private/StaticMeshBuilder.cpp b/Engine/Source/Runtime/MeshBuilder/Private/StaticMeshBuilder.cpp
--- a/Engine/Source/Runtime/MeshBuilder/Private/StaticMeshBuilder.cpp
+++ b/Engine/Source/Runtime/MeshBuilder/Private/StaticMeshBuilder.cpp
@@ -38,7 +38,7 @@ FStaticMeshBuilder::FStaticMeshBuilder()
 
 }
 
-bool FStaticMeshBuilder::Build(UStaticMesh* StaticMesh)
+bool FStaticMeshBuilder::Build(UStaticMesh* StaticMesh, const FStaticMeshLODGroup& LODGroup)
 {
 	if (StaticMesh->GetOriginalMeshDescription(0) == nullptr)
 	{
@@ -47,7 +47,7 @@ bool FStaticMeshBuilder::Build(UStaticMesh* StaticMesh)
 	}
 	StaticMesh->RenderData->AllocateLODResources(StaticMesh->SourceModels.Num());
 
-	OnBuildRenderMeshStart(StaticMesh, false);
+	//OnBuildRenderMeshStart(StaticMesh, false);
 
 	FStaticMeshRenderData& StaticMeshRenderData = *StaticMesh->RenderData;
 	for (int32 LodIndex = 0; LodIndex < StaticMesh->SourceModels.Num(); ++LodIndex)
@@ -60,6 +60,19 @@ bool FStaticMeshBuilder::Build(UStaticMesh* StaticMesh)
 		check(MeshDescription != nullptr);
 		StaticMesh->SetMeshDescription(LodIndex, MeshDescription);
 		
+		//Reduce LODs
+		const FStaticMeshSourceModel& SrcModel = StaticMesh->SourceModels[LodIndex];
+		FMeshReductionSettings ReductionSettings = LODGroup.GetSettings(SrcModel.ReductionSettings, LodIndex);
+		if (ReductionSettings.PercentTriangles < 1.0f || ReductionSettings.MaxDeviation > 0.0f)
+		{
+			float OverlappingThreshold = LODBuildSettings.bRemoveDegenerates ? THRESH_POINTS_ARE_SAME : 0.0f;
+			TMultiMap<int32, int32> OverlappingCorners;
+			MeshDescriptionHelper.FindOverlappingCorners(OverlappingCorners, StaticMesh->GetMeshDescription(LodIndex), OverlappingThreshold);
+			MeshDescriptionHelper.ReduceLOD(StaticMesh->GetMeshDescription(0), StaticMesh->GetMeshDescription(LodIndex), ReductionSettings, OverlappingCorners);
+			// Recompute adjacency information. Since we change the vertices when we reduce
+			MeshDescriptionHelper.FindOverlappingCorners(StaticMesh->GetMeshDescription(0), OverlappingThreshold);
+		}
+
 		const FPolygonGroupArray& PolygonGroups = MeshDescription->PolygonGroups();
 
 		FStaticMeshLODResources& StaticMeshLOD = StaticMeshRenderData.LODResources[LodIndex];
@@ -107,16 +120,29 @@ bool FStaticMeshBuilder::Build(UStaticMesh* StaticMesh)
 		}
 		const EIndexBufferStride::Type IndexBufferStride = bNeeds32BitIndices ? EIndexBufferStride::Force32Bit : EIndexBufferStride::Force16Bit;
 		StaticMeshLOD.IndexBuffer.SetIndices(IndexBuffer, IndexBufferStride);
-		if (MeshDescription->VertexInstances().Num() < 100000 * 3)
-		{
-			BuildOptimizationHelper::CacheOptimizeVertexAndIndexBuffer(StaticMeshBuildVertices, PerSectionIndices, WedgeMap);
-			check(WedgeMap.Num() == MeshDescription->VertexInstances().Num());
-		}
 
 		BuildAllBufferOptimizations(StaticMeshLOD, LODBuildSettings, IndexBuffer, bNeeds32BitIndices, StaticMeshBuildVertices);
 	} //End of LOD for loop
 
-	OnBuildRenderMeshFinish(StaticMesh, true);
+	// Calculate the bounding box.
+	FBox BoundingBox(ForceInit);
+	FPositionVertexBuffer& BasePositionVertexBuffer = StaticMeshRenderData.LODResources[0].PositionVertexBuffer;
+	for (uint32 VertexIndex = 0; VertexIndex < BasePositionVertexBuffer.GetNumVertices(); VertexIndex++)
+	{
+		BoundingBox += BasePositionVertexBuffer.VertexPosition(VertexIndex);
+	}
+	BoundingBox.GetCenterAndExtents(StaticMeshRenderData.Bounds.Origin, StaticMeshRenderData.Bounds.BoxExtent);
+
+	// Calculate the bounding sphere, using the center of the bounding box as the origin.
+	StaticMeshRenderData.Bounds.SphereRadius = 0.0f;
+	for (uint32 VertexIndex = 0; VertexIndex < BasePositionVertexBuffer.GetNumVertices(); VertexIndex++)
+	{
+		StaticMeshRenderData.Bounds.SphereRadius = FMath::Max(
+			(BasePositionVertexBuffer.VertexPosition(VertexIndex) - StaticMeshRenderData.Bounds.Origin).Size(),
+			StaticMeshRenderData.Bounds.SphereRadius
+		);
+	}
+	//OnBuildRenderMeshFinish(StaticMesh, true);
 
 	return true;
 }
@@ -397,7 +423,7 @@ void BuildVertexBuffer(
 				{
 					int32 VertexInstanceValue = Triangle.GetVertexInstanceID(TriVert).GetValue();
 					const FMeshVertexInstance& VertexInstance = MeshDescription->GetVertexInstance(Triangle.GetVertexInstanceID(TriVert));
-					if (VertexInstance.Color != FColor::White)
+					if (VertexInstance.Color != FLinearColor::White)
 					{
 						bHasColor = true;
 					}
@@ -431,12 +457,6 @@ void BuildVertexBuffer(
 					int32 Index = INDEX_NONE;
 					for (int32 k = 0; k < DupVerts.Num(); k++)
 					{
-						if (DupVerts[k] >= VertexInstanceValue)
-						{
-							// the verts beyond me haven't been placed yet, so these duplicates are not relevant
-							break;
-						}
-
 						int32 Location = RemapVerts.IsValidIndex(DupVerts[k]) ? RemapVerts[DupVerts[k]] : INDEX_NONE;
 						if (Location != INDEX_NONE && AreVerticesEqual(StaticMeshVertex, StaticMeshBuildVertices[Location], VertexComparisonThreshold))
 						{
@@ -469,6 +489,13 @@ void BuildVertexBuffer(
 		}
 	}
 
+	//Optimize before setting the buffer
+	if (MeshDescription->VertexInstances().Num() < 100000 * 3)
+	{
+		BuildOptimizationHelper::CacheOptimizeVertexAndIndexBuffer(StaticMeshBuildVertices, OutPerSectionIndices, OutWedgeMap);
+		check(OutWedgeMap.Num() == MeshDescription->VertexInstances().Num());
+	}
+
 	StaticMeshLOD.PositionVertexBuffer.Init(StaticMeshBuildVertices);
 	StaticMeshLOD.VertexBuffer.SetUseHighPrecisionTangentBasis(LODBuildSettings.bUseHighPrecisionTangentBasis);
 	StaticMeshLOD.VertexBuffer.SetUseFullPrecisionUVs(LODBuildSettings.bUseFullPrecisionUVs);
diff --git a/Engine/Source/Runtime/MeshBuilder/Public/IMeshBuilderModule.h b/Engine/Source/Runtime/MeshBuilder/Public/IMeshBuilderModule.h
--- a/Engine/Source/Runtime/MeshBuilder/Public/IMeshBuilderModule.h
+++ b/Engine/Source/Runtime/MeshBuilder/Public/IMeshBuilderModule.h
@@ -31,5 +31,5 @@ public:
 		return FModuleManager::Get().IsModuleLoaded( "MeshBuilder" );
 	}
 
-	virtual bool BuildMesh(class UObject* Mesh) { return false; }
+	virtual bool BuildMesh(class UObject* Mesh, const class FStaticMeshLODGroup& LODGroup) { return false; }
 };
diff --git a/Engine/Source/Runtime/MeshBuilder/Public/MeshBuilder.h b/Engine/Source/Runtime/MeshBuilder/Public/MeshBuilder.h
--- a/Engine/Source/Runtime/MeshBuilder/Public/MeshBuilder.h
+++ b/Engine/Source/Runtime/MeshBuilder/Public/MeshBuilder.h
@@ -17,7 +17,7 @@ public:
 	/**
 	 * Build function should be override and is the starting point for all mesh builders
 	 */
-	virtual bool Build(class UStaticMesh* StaticMesh) =0;
+	virtual bool Build(class UStaticMesh* StaticMesh, const class FStaticMeshLODGroup& LODGroup) =0;
 
 private:
 
diff --git a/Engine/Source/Runtime/MeshBuilder/Public/StaticMeshBuilder.h b/Engine/Source/Runtime/MeshBuilder/Public/StaticMeshBuilder.h
--- a/Engine/Source/Runtime/MeshBuilder/Public/StaticMeshBuilder.h
+++ b/Engine/Source/Runtime/MeshBuilder/Public/StaticMeshBuilder.h
@@ -10,7 +10,7 @@ class MESHBUILDER_API FStaticMeshBuilder : public FMeshBuilder
 public:
 	FStaticMeshBuilder();
 
-	virtual bool Build(class UStaticMesh* StaticMesh) override;
+	virtual bool Build(class UStaticMesh* StaticMesh, const class FStaticMeshLODGroup& LODGroup) override;
 	virtual ~FStaticMeshBuilder() {}
 
 private:
