Fix CIS build

diff --git a/Engine/Source/Developer/MeshSimplifier/Private/QuadricMeshReduction.cpp b/Engine/Source/Developer/MeshSimplifier/Private/QuadricMeshReduction.cpp
--- a/Engine/Source/Developer/MeshSimplifier/Private/QuadricMeshReduction.cpp
+++ b/Engine/Source/Developer/MeshSimplifier/Private/QuadricMeshReduction.cpp
@@ -12,6 +12,7 @@
 #include "IMeshReductionInterfaces.h"
 #include "MeshDescription.h"
 #include "MeshAttributes.h"
+#include "RenderUtils.h"
 
 class FQuadricSimplifierMeshReductionModule : public IMeshReductionModule
 {
diff --git a/Engine/Source/Runtime/Engine/Private/UnrealClient.cpp b/Engine/Source/Runtime/Engine/Private/UnrealClient.cpp
--- a/Engine/Source/Runtime/Engine/Private/UnrealClient.cpp
+++ b/Engine/Source/Runtime/Engine/Private/UnrealClient.cpp
@@ -29,6 +29,7 @@
 #include "EngineModule.h"
 #include "Performance/EnginePerformanceTargets.h"
 #include "Templates/UniquePtr.h"
+#include "EngineUtils.h"
 
 DEFINE_LOG_CATEGORY_STATIC(LogClient, Log, All);
 
diff --git a/Engine/Source/Runtime/MeshDescription/Public/IMeshDescriptionModule.h b/Engine/Source/Runtime/MeshDescription/Public/IMeshDescriptionModule.h
--- a/Engine/Source/Runtime/MeshDescription/Public/IMeshDescriptionModule.h
+++ b/Engine/Source/Runtime/MeshDescription/Public/IMeshDescriptionModule.h
@@ -4,6 +4,7 @@
 
 #include "CoreMinimal.h"
 #include "Modules/ModuleInterface.h"
+#include "Modules/ModuleManager.h"
 
 class IMeshDescriptionModule : public IModuleInterface
 {
