Fix the quadric mesh reduction

diff --git a/Engine/Source/Developer/MeshSimplifier/Private/QuadricMeshReduction.cpp b/Engine/Source/Developer/MeshSimplifier/Private/QuadricMeshReduction.cpp
--- a/Engine/Source/Developer/MeshSimplifier/Private/QuadricMeshReduction.cpp
+++ b/Engine/Source/Developer/MeshSimplifier/Private/QuadricMeshReduction.cpp
@@ -552,6 +552,9 @@ public:
 
 			for (const FMeshTriangle& MeshTriangle : Triangles)
 			{
+				int32 CurrentFaceIndex = FaceIndex;
+				//Increment face index here because there is many continue in this for loop
+				++FaceIndex;
 				for (int32 CornerIndex = 0; CornerIndex < 3; ++CornerIndex)
 				{
 					VertexInstanceIDs[CornerIndex] = MeshTriangle.GetVertexInstanceID(CornerIndex);
@@ -570,8 +573,7 @@ public:
 				int32 VertexIndices[3];
 				for (int32 CornerIndex = 0; CornerIndex < 3; CornerIndex++)
 				{
-					int32 WedgeIndex = FaceIndex * 3 + CornerIndex;
-					++FaceIndex;
+					int32 WedgeIndex = CurrentFaceIndex * 3 + CornerIndex;
 
 					TVertSimp< NumTexCoords > NewVert;
 
diff --git a/Engine/Source/Runtime/Engine/Private/StaticMesh.cpp b/Engine/Source/Runtime/Engine/Private/StaticMesh.cpp
--- a/Engine/Source/Runtime/Engine/Private/StaticMesh.cpp
+++ b/Engine/Source/Runtime/Engine/Private/StaticMesh.cpp
@@ -1350,7 +1350,7 @@ FArchive& operator<<(FArchive& Ar, FMeshBuildSettings& BuildSettings)
 // differences, etc.) replace the version GUID below with a new one.
 // In case of merge conflicts with DDC versions, you MUST generate a new GUID
 // and set this new GUID as the version.                                       
-#define STATICMESH_DERIVEDDATA_VER TEXT("42B9CFE0CAFE46F29F43E9049AA21282")
+#define STATICMESH_DERIVEDDATA_VER TEXT("4FDA6C9D166447EAB2DA9EDA8A184591")
 
 static const FString& GetStaticMeshDerivedDataVersion()
 {
