Make sure reduce LOD iterate the polygons in the same order has FOverlappingCorner

diff --git a/Engine/Source/Developer/MeshSimplifier/Private/QuadricMeshReduction.cpp b/Engine/Source/Developer/MeshSimplifier/Private/QuadricMeshReduction.cpp
--- a/Engine/Source/Developer/MeshSimplifier/Private/QuadricMeshReduction.cpp
+++ b/Engine/Source/Developer/MeshSimplifier/Private/QuadricMeshReduction.cpp
@@ -228,128 +228,126 @@ public:
 
 		TPolygonGroupAttributesRef<FName> OutPolygonGroupMaterialNames = OutReducedMesh.PolygonGroupAttributes().GetAttributesRef<FName>(MeshAttribute::PolygonGroup::ImportedMaterialSlotName);
 
-		int32 FaceIndex = 0;
-		for (FPolygonGroupID PolygonGroupID : InMesh.PolygonGroups().GetElementIDs())
+		int32 WedgeIndex = 0;
+		for (const FPolygonID& PolygonID : InMesh.Polygons().GetElementIDs())
 		{
-			const FMeshPolygonGroup& PolygonGroup = InMesh.GetPolygonGroup(PolygonGroupID);
-			for (const FPolygonID PolygonID : PolygonGroup.Polygons)
+			const FPolygonGroupID PolygonGroupID = InMesh.GetPolygonPolygonGroup(PolygonID);
+
+			const TArray<FMeshTriangle>& PolygonTriangles = InMesh.GetPolygonTriangles(PolygonID);
+			for (int32 TriangleIndex = 0; TriangleIndex < PolygonTriangles.Num(); ++TriangleIndex)
 			{
-				const TArray<FMeshTriangle>& Triangles = InMesh.GetPolygonTriangles(PolygonID);
+				const FMeshTriangle& Triangle = PolygonTriangles[TriangleIndex];
 
-				FVertexInstanceID VertexInstanceIDs[3];
-				FVertexID VertexIDs[3];
-				FVector Positions[3];
+				FVector CornerPositions[3];
+				for (int32 TriVert = 0; TriVert < 3; ++TriVert)
+				{
+					const FVertexInstanceID VertexInstanceID = Triangle.GetVertexInstanceID(TriVert);
+					const FVertexID TmpVertexID = InMesh.GetVertexInstanceVertex(VertexInstanceID);
+					const FVertexID VertexID = bWeldVertices ? VertexIDRemap[TmpVertexID] : TmpVertexID;
+					CornerPositions[TriVert] = InVertexPositions[VertexID];
+				}
 
-				for (const FMeshTriangle MeshTriangle : Triangles)
+				// Don't process degenerate triangles.
+				if( PointsEqual(CornerPositions[0], CornerPositions[1]) ||
+					PointsEqual(CornerPositions[0], CornerPositions[2]) ||
+					PointsEqual(CornerPositions[1], CornerPositions[2]) )
 				{
-					int32 CurrentFaceIndex = FaceIndex;
-					//Increment face index here because there is many continue in this for loop
-					++FaceIndex;
-					for (int32 CornerIndex = 0; CornerIndex < 3; ++CornerIndex)
-					{
-						VertexInstanceIDs[CornerIndex] = MeshTriangle.GetVertexInstanceID(CornerIndex);
-						FVertexID TmpVertexID = InMesh.GetVertexInstanceVertex(VertexInstanceIDs[CornerIndex]);
-						VertexIDs[CornerIndex] = bWeldVertices ? VertexIDRemap[TmpVertexID] : TmpVertexID;
-						Positions[CornerIndex] = InVertexPositions[VertexIDs[CornerIndex]];
-					}
+					WedgeIndex += 3;
+					continue;
+				}
+
+				int32 VertexIndices[3];
+				for (int32 TriVert = 0; TriVert < 3; ++TriVert, ++WedgeIndex)
+				{
+					const FVertexInstanceID VertexInstanceID = Triangle.GetVertexInstanceID(TriVert);
+					const int32 VertexInstanceValue = VertexInstanceID.GetValue();
+					const FVector& VertexPosition = CornerPositions[TriVert];
+
+					TVertSimp< NumTexCoords > NewVert;
 
-					// Don't process degenerate triangles.
-					if (PointsEqual(Positions[0], Positions[1]) ||
-						PointsEqual(Positions[0], Positions[2]) ||
-						PointsEqual(Positions[1], Positions[2]))
+					const TArray<FPolygonID>& VertexInstanceConnectedPolygons = InMesh.GetVertexInstanceConnectedPolygons( VertexInstanceID );
+					if (VertexInstanceConnectedPolygons.Num() > 0)
 					{
-						continue;
+						const FPolygonID ConnectedPolygonID = VertexInstanceConnectedPolygons[0];
+						NewVert.MaterialIndex = InMesh.GetPolygonPolygonGroup(ConnectedPolygonID).GetValue();
 					}
 
-					int32 VertexIndices[3];
-					for (int32 CornerIndex = 0; CornerIndex < 3; CornerIndex++)
+					NewVert.Position = CornerPositions[TriVert];
+					NewVert.Tangents[0] = InVertexTangents[ VertexInstanceID ];
+					NewVert.Normal = InVertexNormals[ VertexInstanceID ];
+					NewVert.Tangents[1] = FVector(0.0f);
+					if (!NewVert.Normal.IsNearlyZero(SMALL_NUMBER) && !NewVert.Tangents[0].IsNearlyZero(SMALL_NUMBER))
 					{
-						int32 WedgeIndex = CurrentFaceIndex * 3 + CornerIndex;
+						NewVert.Tangents[1] = FVector::CrossProduct(NewVert.Normal, NewVert.Tangents[0]).GetSafeNormal() * InVertexBinormalSigns[ VertexInstanceID ];
+					}
 
-						TVertSimp< NumTexCoords > NewVert;
+					// Fix bad tangents
+					NewVert.Tangents[0] = NewVert.Tangents[0].ContainsNaN() ? FVector::ZeroVector : NewVert.Tangents[0];
+					NewVert.Tangents[1] = NewVert.Tangents[1].ContainsNaN() ? FVector::ZeroVector : NewVert.Tangents[1];
+					NewVert.Normal = NewVert.Normal.ContainsNaN() ? FVector::ZeroVector : NewVert.Normal;
+					NewVert.Color = FLinearColor(InVertexColors[ VertexInstanceID ]);
 
-						const TArray<FPolygonID>& VertexInstanceConnectedPolygons = InMesh.GetVertexInstanceConnectedPolygons(VertexInstanceIDs[CornerIndex]);
-						if (VertexInstanceConnectedPolygons.Num() > 0)
+					for (int32 UVIndex = 0; UVIndex < NumTexCoords; UVIndex++)
+					{
+						if (UVIndex < InVertexUVs.GetNumIndices())
 						{
-							const FPolygonID ConnectedPolygonID = VertexInstanceConnectedPolygons[0];
-							NewVert.MaterialIndex = InMesh.GetPolygonPolygonGroup(ConnectedPolygonID).GetValue();
-							// @todo: check with Alexis: OK to conflate material index with polygon group ID? (what if there are gaps in the polygon group array?)
+							NewVert.TexCoords[UVIndex] = InVertexUVs.Get(VertexInstanceID, UVIndex);
+							InMeshNumTexCoords = FMath::Max(UVIndex + 1, InMeshNumTexCoords);
 						}
-
-						NewVert.Position = Positions[CornerIndex];
-						NewVert.Tangents[0] = InVertexTangents[VertexInstanceIDs[CornerIndex]];
-						NewVert.Normal = InVertexNormals[VertexInstanceIDs[CornerIndex]];
-						NewVert.Tangents[1] = FVector(0.0f);
-						if (!NewVert.Normal.IsNearlyZero(SMALL_NUMBER) && !NewVert.Tangents[0].IsNearlyZero(SMALL_NUMBER))
+						else
 						{
-							NewVert.Tangents[1] = FVector::CrossProduct(NewVert.Normal, NewVert.Tangents[0]).GetSafeNormal() * InVertexBinormalSigns[VertexInstanceIDs[CornerIndex]];
+							NewVert.TexCoords[UVIndex] = FVector2D::ZeroVector;
 						}
+					}
+
+					// Make sure this vertex is valid from the start
+					NewVert.Correct();
+					
 
-						// Fix bad tangents
-						NewVert.Tangents[0] = NewVert.Tangents[0].ContainsNaN() ? FVector::ZeroVector : NewVert.Tangents[0];
-						NewVert.Tangents[1] = NewVert.Tangents[1].ContainsNaN() ? FVector::ZeroVector : NewVert.Tangents[1];
-						NewVert.Normal = NewVert.Normal.ContainsNaN() ? FVector::ZeroVector : NewVert.Normal;
-						NewVert.Color = FLinearColor(InVertexColors[VertexInstanceIDs[CornerIndex]]);
+					//Never add duplicated vertex instance
+					//Use WedgeIndex since OverlappingCorners has been built based on that
+					const TArray<int32>& DupVerts = InOverlappingCorners.FindIfOverlapping(WedgeIndex);
 
-						for (int32 UVIndex = 0; UVIndex < NumTexCoords; UVIndex++)
+					int32 Index = INDEX_NONE;
+					for (int32 k = 0; k < DupVerts.Num(); k++)
+					{
+						if (DupVerts[k] >= WedgeIndex)
 						{
-							if (UVIndex < InVertexUVs.GetNumIndices())
-							{
-								NewVert.TexCoords[UVIndex] = InVertexUVs.Get(VertexInstanceIDs[CornerIndex], UVIndex);
-								InMeshNumTexCoords = FMath::Max(UVIndex + 1, InMeshNumTexCoords);
-							}
-							else
-							{
-								NewVert.TexCoords[UVIndex] = FVector2D::ZeroVector;
-							}
+							// the verts beyond me haven't been placed yet, so these duplicates are not relevant
+							break;
 						}
 
-						// Make sure this vertex is valid from the start
-						NewVert.Correct();
-
-						const TArray<int32>& DupVerts = InOverlappingCorners.FindIfOverlapping(WedgeIndex);
-
-						int32 Index = INDEX_NONE;
-						for (int32 k = 0; k < DupVerts.Num(); k++)
+						int32* Location = VertsMap.Find(DupVerts[k]);
+						if (Location)
 						{
-							if (DupVerts[k] >= WedgeIndex)
-							{
-								// the verts beyond me haven't been placed yet, so these duplicates are not relevant
-								break;
-							}
+							TVertSimp< NumTexCoords >& FoundVert = Verts[*Location];
 
-							int32* Location = VertsMap.Find(DupVerts[k]);
-							if (Location)
+							if (NewVert.Equals(FoundVert))
 							{
-								TVertSimp< NumTexCoords >& FoundVert = Verts[*Location];
-
-								if (NewVert.Equals(FoundVert))
-								{
-									Index = *Location;
-									break;
-								}
+								Index = *Location;
+								break;
 							}
 						}
-						if (Index == INDEX_NONE)
-						{
-							Index = Verts.Add(NewVert);
-							VertsMap.Add(WedgeIndex, Index);
-						}
-						VertexIndices[CornerIndex] = Index;
 					}
-
-					// Reject degenerate triangles.
-					if (VertexIndices[0] == VertexIndices[1] ||
-						VertexIndices[1] == VertexIndices[2] ||
-						VertexIndices[0] == VertexIndices[2])
+					if (Index == INDEX_NONE)
 					{
-						continue;
+						Index = Verts.Add(NewVert);
+						VertsMap.Add(WedgeIndex, Index);
 					}
-
-					Indexes.Add(VertexIndices[0]);
-					Indexes.Add(VertexIndices[1]);
-					Indexes.Add(VertexIndices[2]);
+					VertexIndices[TriVert] = Index;
 				}
+				
+				// Reject degenerate triangles.
+				if (VertexIndices[0] == VertexIndices[1] ||
+					VertexIndices[1] == VertexIndices[2] ||
+					VertexIndices[0] == VertexIndices[2])
+				{
+					continue;
+				}
+
+				Indexes.Add(VertexIndices[0]);
+				Indexes.Add(VertexIndices[1]);
+				Indexes.Add(VertexIndices[2]);
 			}
 		}
 
diff --git a/Engine/Source/Runtime/Engine/Private/StaticMesh.cpp b/Engine/Source/Runtime/Engine/Private/StaticMesh.cpp
--- a/Engine/Source/Runtime/Engine/Private/StaticMesh.cpp
+++ b/Engine/Source/Runtime/Engine/Private/StaticMesh.cpp
@@ -1876,7 +1876,7 @@ static void SerializeBuildSettingsForDDC(FArchive& Ar, FMeshBuildSettings& Build
 // differences, etc.) replace the version GUID below with a new one.
 // In case of merge conflicts with DDC versions, you MUST generate a new GUID
 // and set this new GUID as the version.                                       
-#define STATICMESH_DERIVEDDATA_VER TEXT("F36ADCAE95C14BAF91F7B137599924AB")
+#define STATICMESH_DERIVEDDATA_VER TEXT("9891BF2FF72141F6921E24E347DB73C1")
 
 static const FString& GetStaticMeshDerivedDataVersion()
 {
