Deprecate FRawMesh - ReductionInterface

diff --git a/Engine/Source/Developer/MeshReductionInterface/Public/IMeshReductionInterfaces.h b/Engine/Source/Developer/MeshReductionInterface/Public/IMeshReductionInterfaces.h
--- a/Engine/Source/Developer/MeshReductionInterface/Public/IMeshReductionInterfaces.h
+++ b/Engine/Source/Developer/MeshReductionInterface/Public/IMeshReductionInterfaces.h
@@ -8,7 +8,6 @@
 #include "Modules/ModuleInterface.h"
 
 struct FOverlappingCorners;
-struct FRawMesh;
 struct FMeshDescription;
 
 /**
@@ -19,22 +18,6 @@ class IMeshReduction
 public:
 	virtual ~IMeshReduction() {}
 
-	/**
-	* DEPRECATED Please use ReduceMeshDescription instead.
-	* Reduces the raw mesh using the provided reduction settings.
-	* @param OutReducedMesh - Upon return contains the reduced mesh.
-	* @param OutMaxDeviation - Upon return contains the maximum distance by which the reduced mesh deviates from the original.
-	* @param InMesh - The mesh to reduce.
-	* @param ReductionSettings - Setting with which to reduce the mesh.
-	*/
-	virtual void Reduce(
-		FRawMesh& OutReducedMesh,
-		float& OutMaxDeviation,
-		const struct FRawMesh& InMesh,
-		const FOverlappingCorners& InOverlappingCorners,
-		const struct FMeshReductionSettings& ReductionSettings
-	) = 0;
-
 	/**
 	* Reduces the raw mesh using the provided reduction settings.
 	* @param OutReducedMesh - Upon return contains the reduced mesh.
diff --git a/Engine/Source/Developer/MeshSimplifier/Private/QuadricMeshReduction.cpp b/Engine/Source/Developer/MeshSimplifier/Private/QuadricMeshReduction.cpp
--- a/Engine/Source/Developer/MeshSimplifier/Private/QuadricMeshReduction.cpp
+++ b/Engine/Source/Developer/MeshSimplifier/Private/QuadricMeshReduction.cpp
@@ -178,336 +178,6 @@ public:
 		return Version;
 	}
 
-	virtual void Reduce(
-		FRawMesh& OutReducedMesh,
-		float& OutMaxDeviation,
-		const FRawMesh& InMesh,
-		const FOverlappingCorners& InOverlappingCorners,
-		const FMeshReductionSettings& InSettings
-		) override
-	{
-		IMeshUtilities& MeshUtilities = FModuleManager::Get().LoadModuleChecked<IMeshUtilities>("MeshUtilities");
-
-		TArray<FVector> VertexPositions;
-		TArray<uint32> Indices;
-		if (InSettings.WeldingThreshold > 0.0f)
-		{
-			WeldVertexPositions(InMesh, InSettings.WeldingThreshold, VertexPositions, Indices);
-		}
-
-		const uint32 NumTexCoords = MAX_STATIC_TEXCOORDS;
-
-		TArray< TVertSimp< NumTexCoords > >	Verts;
-		TArray< uint32 >					Indexes;
-
-		TMap< int32, int32 > VertsMap;
-
-		int32 NumWedges = InMesh.WedgeIndices.Num();
-		int32 NumFaces = NumWedges / 3;
-
-		// Process each face, build vertex buffer and index buffer
-		for (int32 FaceIndex = 0; FaceIndex < NumFaces; FaceIndex++)
-		{
-			FVector Positions[3];
-			for (int32 CornerIndex = 0; CornerIndex < 3; CornerIndex++)
-			{
-				if (InSettings.WeldingThreshold > 0.0f)
-				{
-					Positions[CornerIndex] = VertexPositions[Indices[FaceIndex * 3 + CornerIndex]];
-				}
-				else
-				{
-					Positions[CornerIndex] = InMesh.VertexPositions[InMesh.WedgeIndices[FaceIndex * 3 + CornerIndex]];
-				}
-			}
-
-			// Don't process degenerate triangles.
-			if( PointsEqual( Positions[0], Positions[1] ) ||
-				PointsEqual( Positions[0], Positions[2] ) ||
-				PointsEqual( Positions[1], Positions[2] ) )
-			{
-				continue;
-			}
-
-			int32 VertexIndices[3];
-			for (int32 CornerIndex = 0; CornerIndex < 3; CornerIndex++)
-			{
-				int32 WedgeIndex = FaceIndex * 3 + CornerIndex;
-
-				TVertSimp< NumTexCoords > NewVert;
-				NewVert.MaterialIndex	= InMesh.FaceMaterialIndices[ FaceIndex ];
-				NewVert.Position		= Positions[ CornerIndex ];
-				NewVert.Tangents[0]		= InMesh.WedgeTangentX[ WedgeIndex ];
-				NewVert.Tangents[1]		= InMesh.WedgeTangentY[ WedgeIndex ];
-				NewVert.Normal			= InMesh.WedgeTangentZ[ WedgeIndex ];
-
-				// Fix bad tangents
-				NewVert.Tangents[0] = NewVert.Tangents[0].ContainsNaN() ? FVector::ZeroVector : NewVert.Tangents[0];
-				NewVert.Tangents[1] = NewVert.Tangents[1].ContainsNaN() ? FVector::ZeroVector : NewVert.Tangents[1];
-				NewVert.Normal		= NewVert.Normal.ContainsNaN()		? FVector::ZeroVector : NewVert.Normal;
-
-				if( InMesh.WedgeColors.Num() == NumWedges )
-				{
-					NewVert.Color = FLinearColor::FromSRGBColor( InMesh.WedgeColors[ WedgeIndex ] );
-				}
-				else
-				{
-					NewVert.Color = FLinearColor::Transparent;
-				}
-
-				for( int32 UVIndex = 0; UVIndex < NumTexCoords; UVIndex++ )
-				{
-					if( InMesh.WedgeTexCoords[ UVIndex ].Num() == NumWedges )
-					{
-						NewVert.TexCoords[ UVIndex ] = InMesh.WedgeTexCoords[ UVIndex ][ WedgeIndex ];
-					}
-					else
-					{
-						NewVert.TexCoords[ UVIndex ] = FVector2D::ZeroVector;
-					}
-				}
-
-				// Make sure this vertex is valid from the start
-				NewVert.Correct();
-
-				const TArray<int32>& DupVerts = InOverlappingCorners.FindIfOverlapping(WedgeIndex);
-
-				int32 Index = INDEX_NONE;
-				for (int32 k = 0; k < DupVerts.Num(); k++)
-				{
-					if( DupVerts[k] >= WedgeIndex )
-					{
-						// the verts beyond me haven't been placed yet, so these duplicates are not relevant
-						break;
-					}
-
-					int32* Location = VertsMap.Find( DupVerts[k] );
-					if( Location )
-					{
-						TVertSimp< NumTexCoords >& FoundVert = Verts[ *Location ];
-
-						if( NewVert.Equals( FoundVert ) )
-						{
-							Index = *Location;
-							break;
-						}
-					}
-				}
-				if( Index == INDEX_NONE )
-				{
-					Index = Verts.Add( NewVert );
-					VertsMap.Add( WedgeIndex, Index );
-				}
-				VertexIndices[ CornerIndex ] = Index;
-			}
-
-			// Reject degenerate triangles.
-			if( VertexIndices[0] == VertexIndices[1] ||
-				VertexIndices[1] == VertexIndices[2] ||
-				VertexIndices[0] == VertexIndices[2] )
-			{
-				continue;
-			}
-
-			Indexes.Add( VertexIndices[0] );
-			Indexes.Add( VertexIndices[1] );
-			Indexes.Add( VertexIndices[2] );
-		}
-
-		uint32 NumVerts = Verts.Num();
-		uint32 NumIndexes = Indexes.Num();
-		uint32 NumTris = NumIndexes / 3;
-
-		static_assert( NumTexCoords == 8, "NumTexCoords changed, fix AttributeWeights" );
-		const uint32 NumAttributes = ( sizeof( TVertSimp< NumTexCoords > ) - sizeof( uint32 ) - sizeof( FVector ) ) / sizeof(float);
-		float AttributeWeights[] =
-		{
-			16.0f, 16.0f, 16.0f,	// Normal
-			0.1f, 0.1f, 0.1f,		// Tangent[0]
-			0.1f, 0.1f, 0.1f,		// Tangent[1]
-			0.1f, 0.1f, 0.1f, 0.1f,	// Color
-			0.5f, 0.5f,				// TexCoord[0]
-			0.5f, 0.5f,				// TexCoord[1]
-			0.5f, 0.5f,				// TexCoord[2]
-			0.5f, 0.5f,				// TexCoord[3]
-			0.5f, 0.5f,				// TexCoord[4]
-			0.5f, 0.5f,				// TexCoord[5]
-			0.5f, 0.5f,				// TexCoord[6]
-			0.5f, 0.5f,				// TexCoord[7]
-		};
-		float* ColorWeights = AttributeWeights + 3 + 3 + 3;
-		float* TexCoordWeights = ColorWeights + 4;
-
-		// Zero out weights that aren't used
-		{
-			if( InMesh.WedgeColors.Num() != NumWedges )
-			{
-				ColorWeights[0] = 0.0f;
-				ColorWeights[1] = 0.0f;
-				ColorWeights[2] = 0.0f;
-				ColorWeights[3] = 0.0f;
-			}
-
-			for( int32 TexCoordIndex = 0; TexCoordIndex < NumTexCoords; TexCoordIndex++ )
-			{
-				if( InMesh.WedgeTexCoords[ TexCoordIndex ].Num() != NumWedges )
-				{
-					TexCoordWeights[ 2 * TexCoordIndex + 0 ] = 0.0f;
-					TexCoordWeights[ 2 * TexCoordIndex + 1 ] = 0.0f;
-				}
-				else if (InMesh.WedgeTexCoords[TexCoordIndex].Num() > 0)
-				{
-					// Normalize TexCoordWeights using min/max TexCoord range, with assumption that value ranges above 2 aren't standard UV values
-
-					float MinVal = +FLT_MAX;
-					float MaxVal = -FLT_MAX;
-
-					for (int32 VertexIndex = 0; VertexIndex < InMesh.WedgeTexCoords[TexCoordIndex].Num(); ++VertexIndex)
-					{
-						MinVal = FMath::Min(MinVal, InMesh.WedgeTexCoords[TexCoordIndex][VertexIndex].X);
-						MinVal = FMath::Min(MinVal, InMesh.WedgeTexCoords[TexCoordIndex][VertexIndex].Y);
-						MaxVal = FMath::Max(MaxVal, InMesh.WedgeTexCoords[TexCoordIndex][VertexIndex].X);
-						MaxVal = FMath::Max(MaxVal, InMesh.WedgeTexCoords[TexCoordIndex][VertexIndex].Y);
-					}
-
-					TexCoordWeights[2 * TexCoordIndex + 0] = 1.0f / FMath::Max(2.0f, MaxVal - MinVal);
-					TexCoordWeights[2 * TexCoordIndex + 1] = 1.0f / FMath::Max(2.0f, MaxVal - MinVal);
-				}
-			}
-		}
-		
-		TMeshSimplifier< TVertSimp< NumTexCoords >, NumAttributes >* MeshSimp = new TMeshSimplifier< TVertSimp< NumTexCoords >, NumAttributes >( Verts.GetData(), NumVerts, Indexes.GetData(), NumIndexes );
-
-		MeshSimp->SetAttributeWeights( AttributeWeights );
-		//MeshSimp->SetBoundaryLocked();
-		MeshSimp->InitCosts();
-
-		float MaxErrorSqr = MeshSimp->SimplifyMesh( MAX_FLT, NumTris * InSettings.PercentTriangles );
-
-		NumVerts = MeshSimp->GetNumVerts();
-		NumTris = MeshSimp->GetNumTris();
-		NumIndexes = NumTris * 3;
-
-		MeshSimp->OutputMesh( Verts.GetData(), Indexes.GetData() );
-		delete MeshSimp;
-
-		//Reorder the face to use the material in the correct order
-		TArray<int32> ReduceMeshUsedMaterialIndex;
-		bool bDoRemap = false;
-		for (uint32 TriIndex = 0; TriIndex < NumTris; TriIndex++)
-		{
-			int32 ReduceMaterialIndex = Verts[Indexes[3 * TriIndex]].MaterialIndex;
-			int32 FinalMaterialIndex = ReduceMeshUsedMaterialIndex.AddUnique(ReduceMaterialIndex);
-			bDoRemap |= ReduceMaterialIndex != FinalMaterialIndex;
-		}
-		if (bDoRemap)
-		{
-			int32 MaximumIndex = 0;
-			int32 MaxMaterialIndex = FMath::Max(ReduceMeshUsedMaterialIndex, &MaximumIndex);
-			TArray<TArray<int32>> MaterialSectionIndexes;
-			//We need to add up to the maximum material index
-			MaterialSectionIndexes.AddDefaulted(MaxMaterialIndex+1);
-			//Reorder the Indexes according to the remap array
-			//First, sort them by section in the material index order
-			for (uint32 IndexesIndex = 0; IndexesIndex < NumIndexes; IndexesIndex++)
-			{
-				int32 ReduceMaterialIndex = Verts[Indexes[IndexesIndex]].MaterialIndex;
-				MaterialSectionIndexes[ReduceMaterialIndex].Add(Indexes[IndexesIndex]);
-			}
-			//Update the Indexes array by placing all triangles in section index order
-			//This will make sure that the reduce LOD mesh will have the same material order
-			//as the reference LOD, even if some sections disappear because all triangles was remove.
-			int32 IndexOffset = 0;
-			for (const TArray<int32>& RemapSectionIndexes : MaterialSectionIndexes)
-			{
-				for (int32 IndexOfIndex = 0; IndexOfIndex < RemapSectionIndexes.Num(); ++IndexOfIndex)
-				{
-					int32 SortedIndex = IndexOfIndex + IndexOffset;
-					Indexes[SortedIndex] = RemapSectionIndexes[IndexOfIndex];
-				}
-				IndexOffset += RemapSectionIndexes.Num();
-			}
-		}
-
-
-		OutMaxDeviation = FMath::Sqrt( MaxErrorSqr ) / 8.0f;
-
-		{
-			// Output FRawMesh
-			OutReducedMesh.VertexPositions.Empty( NumVerts );
-			OutReducedMesh.VertexPositions.AddUninitialized( NumVerts );
-			for( uint32 i= 0; i < NumVerts; i++ )
-			{
-				OutReducedMesh.VertexPositions[i] = Verts[i].Position;
-			}
-
-			OutReducedMesh.WedgeIndices.Empty( NumIndexes );
-			OutReducedMesh.WedgeIndices.AddUninitialized( NumIndexes );
-
-			for( uint32 i = 0; i < NumIndexes; i++ )
-			{
-				OutReducedMesh.WedgeIndices[i] = Indexes[i];
-			}
-
-			OutReducedMesh.WedgeTangentX.Empty( NumIndexes );
-			OutReducedMesh.WedgeTangentY.Empty( NumIndexes );
-			OutReducedMesh.WedgeTangentZ.Empty( NumIndexes );
-			OutReducedMesh.WedgeTangentX.AddUninitialized( NumIndexes );
-			OutReducedMesh.WedgeTangentY.AddUninitialized( NumIndexes );
-			OutReducedMesh.WedgeTangentZ.AddUninitialized( NumIndexes );
-			for( uint32 i= 0; i < NumIndexes; i++ )
-			{
-				OutReducedMesh.WedgeTangentX[i] = Verts[ Indexes[i] ].Tangents[0];
-				OutReducedMesh.WedgeTangentY[i] = Verts[ Indexes[i] ].Tangents[1];
-				OutReducedMesh.WedgeTangentZ[i] = Verts[ Indexes[i] ].Normal;
-			}
-
-			if( InMesh.WedgeColors.Num() == NumWedges )
-			{
-				OutReducedMesh.WedgeColors.Empty( NumIndexes );
-				OutReducedMesh.WedgeColors.AddUninitialized( NumIndexes );
-				for( uint32 i= 0; i < NumIndexes; i++ )
-				{
-					OutReducedMesh.WedgeColors[i] = Verts[ Indexes[i] ].Color.ToFColor(true);
-				}
-			}
-			else
-			{
-				OutReducedMesh.WedgeColors.Empty();
-			}
-
-			for( int32 TexCoordIndex = 0; TexCoordIndex < NumTexCoords; TexCoordIndex++ )
-			{
-				if( InMesh.WedgeTexCoords[ TexCoordIndex ].Num() == NumWedges )
-				{
-					OutReducedMesh.WedgeTexCoords[ TexCoordIndex ].Empty( NumIndexes );
-					OutReducedMesh.WedgeTexCoords[ TexCoordIndex ].AddUninitialized( NumIndexes );
-					for( uint32 i= 0; i < NumIndexes; i++ )
-					{
-						OutReducedMesh.WedgeTexCoords[ TexCoordIndex ][i] = Verts[ Indexes[i] ].TexCoords[ TexCoordIndex ];
-					}
-				}
-				else
-				{
-					OutReducedMesh.WedgeTexCoords[ TexCoordIndex ].Empty();
-				}
-			}
-
-			OutReducedMesh.FaceMaterialIndices.Empty( NumTris );
-			OutReducedMesh.FaceMaterialIndices.AddUninitialized( NumTris );
-			for( uint32 i= 0; i < NumTris; i++ )
-			{
-				OutReducedMesh.FaceMaterialIndices[i] = Verts[ Indexes[3*i] ].MaterialIndex;
-			}
-
-			OutReducedMesh.FaceSmoothingMasks.Empty( NumTris );
-			OutReducedMesh.FaceSmoothingMasks.AddZeroed( NumTris );
-
-			Verts.Empty();
-			Indexes.Empty();
-		}
-	}
-
 	virtual void ReduceMeshDescription(
 		FMeshDescription& OutReducedMesh,
 		float& OutMaxDeviation,
diff --git a/Engine/Source/Developer/MeshUtilities/MeshUtilities.Build.cs b/Engine/Source/Developer/MeshUtilities/MeshUtilities.Build.cs
--- a/Engine/Source/Developer/MeshUtilities/MeshUtilities.Build.cs
+++ b/Engine/Source/Developer/MeshUtilities/MeshUtilities.Build.cs
@@ -39,6 +39,8 @@ public class MeshUtilities : ModuleRules
                 "GraphColor",
                 "MeshBuilder",
                 "MeshUtilitiesCommon",
+                "MeshDescription",
+                "MeshDescriptionOperations"
             }
 		);
 
diff --git a/Engine/Source/Developer/MeshUtilities/Private/MeshUtilities.cpp b/Engine/Source/Developer/MeshUtilities/Private/MeshUtilities.cpp
--- a/Engine/Source/Developer/MeshUtilities/Private/MeshUtilities.cpp
+++ b/Engine/Source/Developer/MeshUtilities/Private/MeshUtilities.cpp
@@ -91,6 +91,9 @@
 #include "OverlappingCorners.h"
 #include "MeshUtilitiesCommon.h"
 
+#include "MeshDescription.h"
+#include "MeshDescriptionOperations.h"
+
 #if WITH_EDITOR
 #include "Editor.h"
 #include "UnrealEdMisc.h"
@@ -2578,8 +2581,16 @@ public:
 				FRawMesh& DestMesh = LODMeshes[NumValidLODs];
 				FOverlappingCorners& InOverlappingCorners = LODOverlappingCorners[ReductionSettings.BaseLODModel];
 				FOverlappingCorners& DestOverlappingCorners = LODOverlappingCorners[NumValidLODs];
+				FMeshDescription SrcMeshdescription;
+				UStaticMesh::RegisterMeshAttributes(SrcMeshdescription);
+				FMeshDescription DestMeshdescription;
+				UStaticMesh::RegisterMeshAttributes(DestMeshdescription);
+				TMap<int32, FName> FromMaterialMap;
+				FMeshDescriptionOperations::ConvertFromRawMesh(InMesh, SrcMeshdescription, FromMaterialMap);
+				MeshReduction->ReduceMeshDescription(DestMeshdescription, LODMaxDeviation[NumValidLODs], SrcMeshdescription, InOverlappingCorners, ReductionSettings);
+				TMap<FName, int32> ToMaterialMap;
+				FMeshDescriptionOperations::ConvertToRawMesh(DestMeshdescription, DestMesh, ToMaterialMap);
 
-				MeshReduction->Reduce(DestMesh, LODMaxDeviation[NumValidLODs], InMesh, InOverlappingCorners, ReductionSettings);
 				if (DestMesh.WedgeIndices.Num() > 0 && !DestMesh.IsValid())
 				{
 					UE_LOG(LogMeshUtilities, Error, TEXT("Mesh reduction produced a corrupt mesh for LOD%d"), LODIndex);
diff --git a/Engine/Source/Developer/SimplygonMeshReduction/Private/SimplygonMeshReduction.cpp b/Engine/Source/Developer/SimplygonMeshReduction/Private/SimplygonMeshReduction.cpp
--- a/Engine/Source/Developer/SimplygonMeshReduction/Private/SimplygonMeshReduction.cpp
+++ b/Engine/Source/Developer/SimplygonMeshReduction/Private/SimplygonMeshReduction.cpp
@@ -184,52 +184,6 @@ public:
 		return FString("SimplygonMeshReduction");
 	}
 
-	virtual void Reduce(
-		FRawMesh& OutReducedMesh,
-		float& OutMaxDeviation,
-		const FRawMesh& InMesh,
-		const FOverlappingCorners& InOverlappingCorners,
-		const FMeshReductionSettings& InSettings
-		) override
-	{
-		SimplygonSDK::spGeometryData GeometryData = CreateGeometryFromRawMesh(InMesh);
-		check(GeometryData);
-
-		SimplygonSDK::spScene Scene = SDK->CreateScene();
-
-		SimplygonSDK::spSceneMesh Mesh = SDK->CreateSceneMesh();
-		Mesh->SetGeometry(GeometryData);
-		Mesh->SetName(TCHAR_TO_ANSI(*FString::Printf(TEXT("UnrealMesh"))));
-		Scene->GetRootNode()->AddChild(Mesh);
-
-		SimplygonSDK::spReductionProcessor ReductionProcessor = SDK->CreateReductionProcessor();
-		ReductionProcessor->AddObserver(&EventHandler, SimplygonSDK::SG_EVENT_PROGRESS);
-		ReductionProcessor->SetScene(Scene);
-
-		SimplygonSDK::spRepairSettings RepairSettings = ReductionProcessor->GetRepairSettings();
-		RepairSettings->SetWeldDist(InSettings.WeldingThreshold);
-		RepairSettings->SetTjuncDist(InSettings.WeldingThreshold);
-
-		SimplygonSDK::spReductionSettings ReductionSettings = ReductionProcessor->GetReductionSettings();
-		SetReductionSettings(ReductionSettings, InSettings, GeometryData->GetTriangleCount());
-
-		//Set visibility settings
-		SimplygonSDK::spVisibilitySettings VisibilitySettings = ReductionProcessor->GetVisibilitySettings();
-		VisibilitySettings->SetCullOccludedGeometry(InSettings.bCullOccluded);
-		int32 TempAggressiveness = InSettings.VisibilityAggressiveness + 1; //+1 because there is an offset in aggressiveness options
-		VisibilitySettings->SetVisibilityWeightsPower(TempAggressiveness);
-		VisibilitySettings->SetUseVisibilityWeightsInReducer(InSettings.bVisibilityAided);
-
-		SimplygonSDK::spNormalCalculationSettings NormalSettings = ReductionProcessor->GetNormalCalculationSettings();
-		SetNormalSettings(NormalSettings, InSettings);
-
-		ReductionProcessor->RunProcessing();
-
-		SimplygonSDK::spSceneMesh ReducedMesh = SimplygonSDK::Cast<SimplygonSDK::ISceneMesh>(Scene->GetRootNode()->GetChild(0));
-		CreateRawMeshFromGeometry(OutReducedMesh, ReducedMesh->GetGeometry(), WINDING_Keep);
-		OutMaxDeviation = ReductionProcessor->GetResultDeviation();
-	}
-
 	bool ReduceLODModel(
 		FSkeletalMeshLODModel * SrcModel,
 		FSkeletalMeshLODModel *& OutModel,
