Cleanup the SectionInfoMap when we remove LODs

diff --git a/Engine/Source/Developer/MeshBuilder/Private/StaticMeshBuilder.cpp b/Engine/Source/Developer/MeshBuilder/Private/StaticMeshBuilder.cpp
--- a/Engine/Source/Developer/MeshBuilder/Private/StaticMeshBuilder.cpp
+++ b/Engine/Source/Developer/MeshBuilder/Private/StaticMeshBuilder.cpp
@@ -134,6 +134,8 @@ bool FStaticMeshBuilder::Build(FStaticMeshRenderData& StaticMeshRenderData, USta
 			//Create a reduced mesh from the base LOD
 			UStaticMesh::RegisterMeshAttributes(MeshDescriptions[LodIndex]);
 			
+			int32 OldSectionInfoMapCount = StaticMesh->SectionInfoMap.GetSectionNumber(LodIndex);
+
 			if (LodIndex == BaseLodIndex)
 			{
 				//When using LOD 0, we use a copy of the mesh description since reduce do not support inline reducing
@@ -167,11 +169,15 @@ bool FStaticMeshBuilder::Build(FStaticMeshRenderData& StaticMeshRenderData, USta
 				}
 				UniqueMaterialIndex.AddUnique(MaterialIndex);
 			}
+
+			//If the reduce did not output the same number of section use the base LOD sectionInfoMap
+			bool bIsOldMappingInvalid = OldSectionInfoMapCount != UniqueMaterialIndex.Num();
+
 			//All used material represent a different section
 			for (int32 SectionIndex = 0; SectionIndex < UniqueMaterialIndex.Num(); ++SectionIndex)
 			{
 				//Keep the old data
-				bool bHasValidLODInfoMap = LODModelSectionInfoMap.IsValidSection(LodIndex, SectionIndex);
+				bool bHasValidLODInfoMap = !bIsOldMappingInvalid && LODModelSectionInfoMap.IsValidSection(LodIndex, SectionIndex);
 				//Section material index have to be remap with the ReductionSettings.BaseLODModel SectionInfoMap to create
 				//a valid new section info map for the reduced LOD.
 
@@ -182,7 +188,7 @@ bool FStaticMeshBuilder::Build(FStaticMeshRenderData& StaticMeshRenderData, USta
 					for (int32 BaseLodSectionIndex = 0; BaseLodSectionIndex < LODModelSectionInfoMap.GetSectionNumber(BaseLodIndex); ++BaseLodSectionIndex)
 					{
 						FMeshSectionInfo SectionInfo = LODModelSectionInfoMap.Get(ReductionSettings.BaseLODModel, BaseLodSectionIndex);
-						if (SectionInfo.MaterialIndex == UniqueMaterialIndex[SectionIndex])
+						if (BaseLodSectionIndex == UniqueMaterialIndex[SectionIndex])
 						{
 							//Copy the BaseLODModel section info to the reduce LODIndex.
 							FMeshSectionInfo OriginalSectionInfo = LODModelOriginalSectionInfoMap.Get(ReductionSettings.BaseLODModel, BaseLodSectionIndex);
diff --git a/Engine/Source/Developer/MeshSimplifier/Private/QuadricMeshReduction.cpp b/Engine/Source/Developer/MeshSimplifier/Private/QuadricMeshReduction.cpp
--- a/Engine/Source/Developer/MeshSimplifier/Private/QuadricMeshReduction.cpp
+++ b/Engine/Source/Developer/MeshSimplifier/Private/QuadricMeshReduction.cpp
@@ -223,123 +223,127 @@ public:
 		TPolygonGroupAttributesRef<FName> OutPolygonGroupMaterialNames = OutReducedMesh.PolygonGroupAttributes().GetAttributesRef<FName>(MeshAttribute::PolygonGroup::ImportedMaterialSlotName);
 
 		int32 FaceIndex = 0;
-		for (const FPolygonID& PolygonID : InMesh.Polygons().GetElementIDs())
+		for (FPolygonGroupID PolygonGroupID : InMesh.PolygonGroups().GetElementIDs())
 		{
-			const TArray<FMeshTriangle>& Triangles = InMesh.GetPolygonTriangles(PolygonID);
-
-			FVertexInstanceID VertexInstanceIDs[3];
-			FVertexID VertexIDs[3];
-			FVector Positions[3];
-
-			for (const FMeshTriangle MeshTriangle : Triangles)
+			const FMeshPolygonGroup& PolygonGroup = InMesh.GetPolygonGroup(PolygonGroupID);
+			for (const FPolygonID PolygonID : PolygonGroup.Polygons)
 			{
-				int32 CurrentFaceIndex = FaceIndex;
-				//Increment face index here because there is many continue in this for loop
-				++FaceIndex;
-				for (int32 CornerIndex = 0; CornerIndex < 3; ++CornerIndex)
-				{
-					VertexInstanceIDs[CornerIndex] = MeshTriangle.GetVertexInstanceID(CornerIndex);
-					FVertexID TmpVertexID = InMesh.GetVertexInstanceVertex(VertexInstanceIDs[CornerIndex]);
-					VertexIDs[CornerIndex] = bWeldVertices ? VertexIDRemap[TmpVertexID] : TmpVertexID;
-					Positions[CornerIndex] = InVertexPositions[VertexIDs[CornerIndex]];
-				}
+				const TArray<FMeshTriangle>& Triangles = InMesh.GetPolygonTriangles(PolygonID);
 
-				// Don't process degenerate triangles.
-				if (PointsEqual(Positions[0], Positions[1]) ||
-					PointsEqual(Positions[0], Positions[2]) ||
-					PointsEqual(Positions[1], Positions[2]))
-				{
-					continue;
-				}
+				FVertexInstanceID VertexInstanceIDs[3];
+				FVertexID VertexIDs[3];
+				FVector Positions[3];
 
-				int32 VertexIndices[3];
-				for (int32 CornerIndex = 0; CornerIndex < 3; CornerIndex++)
+				for (const FMeshTriangle MeshTriangle : Triangles)
 				{
-					int32 WedgeIndex = CurrentFaceIndex * 3 + CornerIndex;
-
-					TVertSimp< NumTexCoords > NewVert;
-
-					const TArray<FPolygonID>& VertexInstanceConnectedPolygons = InMesh.GetVertexInstanceConnectedPolygons(VertexInstanceIDs[CornerIndex]);
-					if (VertexInstanceConnectedPolygons.Num() > 0)
+					int32 CurrentFaceIndex = FaceIndex;
+					//Increment face index here because there is many continue in this for loop
+					++FaceIndex;
+					for (int32 CornerIndex = 0; CornerIndex < 3; ++CornerIndex)
 					{
-						const FPolygonID ConnectedPolygonID = VertexInstanceConnectedPolygons[0];
-						NewVert.MaterialIndex = InMesh.GetPolygonPolygonGroup(ConnectedPolygonID).GetValue();
-						// @todo: check with Alexis: OK to conflate material index with polygon group ID? (what if there are gaps in the polygon group array?)
+						VertexInstanceIDs[CornerIndex] = MeshTriangle.GetVertexInstanceID(CornerIndex);
+						FVertexID TmpVertexID = InMesh.GetVertexInstanceVertex(VertexInstanceIDs[CornerIndex]);
+						VertexIDs[CornerIndex] = bWeldVertices ? VertexIDRemap[TmpVertexID] : TmpVertexID;
+						Positions[CornerIndex] = InVertexPositions[VertexIDs[CornerIndex]];
 					}
 
-					NewVert.Position = Positions[CornerIndex];
-					NewVert.Tangents[0] = InVertexTangents[VertexInstanceIDs[CornerIndex]];
-					NewVert.Normal = InVertexNormals[VertexInstanceIDs[CornerIndex]];
-					NewVert.Tangents[1] = FVector(0.0f);
-					if (!NewVert.Normal.IsNearlyZero(SMALL_NUMBER) && !NewVert.Tangents[0].IsNearlyZero(SMALL_NUMBER))
+					// Don't process degenerate triangles.
+					if (PointsEqual(Positions[0], Positions[1]) ||
+						PointsEqual(Positions[0], Positions[2]) ||
+						PointsEqual(Positions[1], Positions[2]))
 					{
-						NewVert.Tangents[1] = FVector::CrossProduct(NewVert.Normal, NewVert.Tangents[0]).GetSafeNormal() * InVertexBinormalSigns[VertexInstanceIDs[CornerIndex]];
+						continue;
 					}
 
-					// Fix bad tangents
-					NewVert.Tangents[0] = NewVert.Tangents[0].ContainsNaN() ? FVector::ZeroVector : NewVert.Tangents[0];
-					NewVert.Tangents[1] = NewVert.Tangents[1].ContainsNaN() ? FVector::ZeroVector : NewVert.Tangents[1];
-					NewVert.Normal = NewVert.Normal.ContainsNaN() ? FVector::ZeroVector : NewVert.Normal;
-					NewVert.Color = FLinearColor(InVertexColors[VertexInstanceIDs[CornerIndex]]);
-
-					for (int32 UVIndex = 0; UVIndex < NumTexCoords; UVIndex++)
+					int32 VertexIndices[3];
+					for (int32 CornerIndex = 0; CornerIndex < 3; CornerIndex++)
 					{
-						if (UVIndex < InVertexUVs.GetNumIndices())
+						int32 WedgeIndex = CurrentFaceIndex * 3 + CornerIndex;
+
+						TVertSimp< NumTexCoords > NewVert;
+
+						const TArray<FPolygonID>& VertexInstanceConnectedPolygons = InMesh.GetVertexInstanceConnectedPolygons(VertexInstanceIDs[CornerIndex]);
+						if (VertexInstanceConnectedPolygons.Num() > 0)
 						{
-							NewVert.TexCoords[UVIndex] = InVertexUVs.Get(VertexInstanceIDs[CornerIndex], UVIndex);
-							InMeshNumTexCoords = FMath::Max(UVIndex+1, InMeshNumTexCoords);
+							const FPolygonID ConnectedPolygonID = VertexInstanceConnectedPolygons[0];
+							NewVert.MaterialIndex = InMesh.GetPolygonPolygonGroup(ConnectedPolygonID).GetValue();
+							// @todo: check with Alexis: OK to conflate material index with polygon group ID? (what if there are gaps in the polygon group array?)
 						}
-						else
+
+						NewVert.Position = Positions[CornerIndex];
+						NewVert.Tangents[0] = InVertexTangents[VertexInstanceIDs[CornerIndex]];
+						NewVert.Normal = InVertexNormals[VertexInstanceIDs[CornerIndex]];
+						NewVert.Tangents[1] = FVector(0.0f);
+						if (!NewVert.Normal.IsNearlyZero(SMALL_NUMBER) && !NewVert.Tangents[0].IsNearlyZero(SMALL_NUMBER))
 						{
-							NewVert.TexCoords[UVIndex] = FVector2D::ZeroVector;
+							NewVert.Tangents[1] = FVector::CrossProduct(NewVert.Normal, NewVert.Tangents[0]).GetSafeNormal() * InVertexBinormalSigns[VertexInstanceIDs[CornerIndex]];
 						}
-					}
 
-					// Make sure this vertex is valid from the start
-					NewVert.Correct();
+						// Fix bad tangents
+						NewVert.Tangents[0] = NewVert.Tangents[0].ContainsNaN() ? FVector::ZeroVector : NewVert.Tangents[0];
+						NewVert.Tangents[1] = NewVert.Tangents[1].ContainsNaN() ? FVector::ZeroVector : NewVert.Tangents[1];
+						NewVert.Normal = NewVert.Normal.ContainsNaN() ? FVector::ZeroVector : NewVert.Normal;
+						NewVert.Color = FLinearColor(InVertexColors[VertexInstanceIDs[CornerIndex]]);
 
-					const TArray<int32>& DupVerts = InOverlappingCorners.FindIfOverlapping(WedgeIndex);
-
-					int32 Index = INDEX_NONE;
-					for (int32 k = 0; k < DupVerts.Num(); k++)
-					{
-						if (DupVerts[k] >= WedgeIndex)
+						for (int32 UVIndex = 0; UVIndex < NumTexCoords; UVIndex++)
 						{
-							// the verts beyond me haven't been placed yet, so these duplicates are not relevant
-							break;
+							if (UVIndex < InVertexUVs.GetNumIndices())
+							{
+								NewVert.TexCoords[UVIndex] = InVertexUVs.Get(VertexInstanceIDs[CornerIndex], UVIndex);
+								InMeshNumTexCoords = FMath::Max(UVIndex + 1, InMeshNumTexCoords);
+							}
+							else
+							{
+								NewVert.TexCoords[UVIndex] = FVector2D::ZeroVector;
+							}
 						}
 
-						int32* Location = VertsMap.Find(DupVerts[k]);
-						if (Location)
-						{
-							TVertSimp< NumTexCoords >& FoundVert = Verts[*Location];
+						// Make sure this vertex is valid from the start
+						NewVert.Correct();
 
-							if (NewVert.Equals(FoundVert))
+						const TArray<int32>& DupVerts = InOverlappingCorners.FindIfOverlapping(WedgeIndex);
+
+						int32 Index = INDEX_NONE;
+						for (int32 k = 0; k < DupVerts.Num(); k++)
+						{
+							if (DupVerts[k] >= WedgeIndex)
 							{
-								Index = *Location;
+								// the verts beyond me haven't been placed yet, so these duplicates are not relevant
 								break;
 							}
+
+							int32* Location = VertsMap.Find(DupVerts[k]);
+							if (Location)
+							{
+								TVertSimp< NumTexCoords >& FoundVert = Verts[*Location];
+
+								if (NewVert.Equals(FoundVert))
+								{
+									Index = *Location;
+									break;
+								}
+							}
+						}
+						if (Index == INDEX_NONE)
+						{
+							Index = Verts.Add(NewVert);
+							VertsMap.Add(WedgeIndex, Index);
 						}
+						VertexIndices[CornerIndex] = Index;
 					}
-					if (Index == INDEX_NONE)
+
+					// Reject degenerate triangles.
+					if (VertexIndices[0] == VertexIndices[1] ||
+						VertexIndices[1] == VertexIndices[2] ||
+						VertexIndices[0] == VertexIndices[2])
 					{
-						Index = Verts.Add(NewVert);
-						VertsMap.Add(WedgeIndex, Index);
+						continue;
 					}
-					VertexIndices[CornerIndex] = Index;
-				}
 
-				// Reject degenerate triangles.
-				if (VertexIndices[0] == VertexIndices[1] ||
-					VertexIndices[1] == VertexIndices[2] ||
-					VertexIndices[0] == VertexIndices[2])
-				{
-					continue;
+					Indexes.Add(VertexIndices[0]);
+					Indexes.Add(VertexIndices[1]);
+					Indexes.Add(VertexIndices[2]);
 				}
-
-				Indexes.Add(VertexIndices[0]);
-				Indexes.Add(VertexIndices[1]);
-				Indexes.Add(VertexIndices[2]);
 			}
 		}
 
diff --git a/Engine/Source/Runtime/Engine/Private/StaticMesh.cpp b/Engine/Source/Runtime/Engine/Private/StaticMesh.cpp
--- a/Engine/Source/Runtime/Engine/Private/StaticMesh.cpp
+++ b/Engine/Source/Runtime/Engine/Private/StaticMesh.cpp
@@ -2266,6 +2266,24 @@ void UStaticMesh::SetNumSourceModels(const int32 Num)
 	SourceModels.SetNum(Num);
 	MeshDescriptions->SetNum(Num);
 
+	//Shrink the SectionInfoMap if some SourceModel are removed
+	if (OldNum > Num)
+	{
+		for (int32 RemoveLODIndex = Num; RemoveLODIndex < OldNum; ++RemoveLODIndex)
+		{
+			int32 SectionCount = SectionInfoMap.GetSectionNumber(RemoveLODIndex);
+			for (int32 SectionIndex = 0; SectionIndex < SectionCount; ++SectionIndex)
+			{
+				SectionInfoMap.Remove(RemoveLODIndex, SectionIndex);
+			}
+			SectionCount = OriginalSectionInfoMap.GetSectionNumber(RemoveLODIndex);
+			for (int32 SectionIndex = 0; SectionIndex < SectionCount; ++SectionIndex)
+			{
+				OriginalSectionInfoMap.Remove(RemoveLODIndex, SectionIndex);
+			}
+		}
+	}
+
 	for (int32 Index = OldNum; Index < Num; ++Index)
 	{
 		SourceModels[Index].StaticMeshOwner = this;
@@ -2277,6 +2295,44 @@ void UStaticMesh::RemoveSourceModel(const int32 Index)
 	LoadMeshDescriptions();
 	check(MeshDescriptions->Num() == SourceModels.Num());
 	check(SourceModels.IsValidIndex(Index));
+	
+	//Remove the SectionInfoMap of the LOD we remove
+	{
+		int32 SectionCount = SectionInfoMap.GetSectionNumber(Index);
+		for (int32 SectionIndex = 0; SectionIndex < SectionCount; ++SectionIndex)
+		{
+			SectionInfoMap.Remove(Index, SectionIndex);
+		}
+		SectionCount = OriginalSectionInfoMap.GetSectionNumber(Index);
+		for (int32 SectionIndex = 0; SectionIndex < SectionCount; ++SectionIndex)
+		{
+			OriginalSectionInfoMap.Remove(Index, SectionIndex);
+		}
+	}
+
+	//Move down all SectionInfoMap for the next LOD
+	if (Index < SourceModels.Num() - 1)
+	{
+		for (int32 MoveIndex = Index + 1; MoveIndex < SourceModels.Num(); ++MoveIndex)
+		{
+			int32 SectionCount = SectionInfoMap.GetSectionNumber(MoveIndex);
+			for (int32 SectionIndex = 0; SectionIndex < SectionCount; ++SectionIndex)
+			{
+				FMeshSectionInfo SectionInfo = SectionInfoMap.Get(MoveIndex, SectionIndex);
+				SectionInfoMap.Set(MoveIndex - 1, SectionIndex, SectionInfo);
+				SectionInfoMap.Remove(MoveIndex, SectionIndex);
+			}
+			SectionCount = OriginalSectionInfoMap.GetSectionNumber(MoveIndex);
+			for (int32 SectionIndex = 0; SectionIndex < SectionCount; ++SectionIndex)
+			{
+				FMeshSectionInfo SectionInfo = OriginalSectionInfoMap.Get(MoveIndex, SectionIndex);
+				OriginalSectionInfoMap.Set(MoveIndex - 1, SectionIndex, SectionInfo);
+				OriginalSectionInfoMap.Remove(MoveIndex, SectionIndex);
+			}
+		}
+	}
+	
+	//Remove the LOD
 	SourceModels.RemoveAt(Index);
 	MeshDescriptions->RemoveAt(Index);
 }
