Static Mesh Reduction:  Make more robust to extreme UV values on the source mesh.

diff --git a/Engine/Source/Developer/MeshSimplifier/Private/QuadricMeshReduction.cpp b/Engine/Source/Developer/MeshSimplifier/Private/QuadricMeshReduction.cpp
--- a/Engine/Source/Developer/MeshSimplifier/Private/QuadricMeshReduction.cpp
+++ b/Engine/Source/Developer/MeshSimplifier/Private/QuadricMeshReduction.cpp
@@ -371,6 +371,47 @@ public:
 		float* ColorWeights = AttributeWeights + 3 + 3 + 3;
 		float* TexCoordWeights = ColorWeights + 4;
 
+		// Re-scale the weights for UV channels that exceed the expected 0-1 range.
+		// Otherwise garbage on the UVs will dominate the simplification quadric.
+		{
+			float XLength[MAX_STATIC_TEXCOORDS] = { 0 };
+			float YLength[MAX_STATIC_TEXCOORDS] = { 0 };
+			{
+				for (int32 TexCoordId = 0; TexCoordId < NumTexCoords; ++TexCoordId)
+				{
+					float XMax = -FLT_MAX;
+					float YMax = -FLT_MAX;
+					float XMin = FLT_MAX;
+					float YMin = FLT_MAX;
+					for (const TVertSimp< NumTexCoords >& SimpVert : Verts)
+					{
+						const FVector2D& UVs = SimpVert.TexCoords[TexCoordId];
+						XMax = FMath::Max(XMax, UVs.X);
+						XMin = FMath::Min(XMin, UVs.X);
+
+						YMax = FMath::Max(YMax, UVs.Y);
+						YMin = FMath::Min(YMin, UVs.Y);
+					}
+
+					XLength[TexCoordId] =  ( XMax > XMin ) ? XMax - XMin : 0.f;
+					YLength[TexCoordId] =  ( YMax > YMin ) ? YMax - YMin : 0.f;
+				}
+			}
+
+			for (int32 TexCoordId = 0; TexCoordId < NumTexCoords; ++TexCoordId)
+			{
+
+				if (XLength[TexCoordId] > 1.f)
+				{
+					TexCoordWeights[2 * TexCoordId + 0] /= XLength[TexCoordId];
+				}
+				if (YLength[TexCoordId] > 1.f)
+				{
+					TexCoordWeights[2 * TexCoordId + 1] /= YLength[TexCoordId];
+				}
+			}
+		}
+
 		// Zero out weights that aren't used
 		{
 			//TODO Check if we have vertex color
