Fix for clang 9

diff --git a/Engine/Source/Developer/MeshSimplifier/Private/Quadric.cpp b/Engine/Source/Developer/MeshSimplifier/Private/Quadric.cpp
--- a/Engine/Source/Developer/MeshSimplifier/Private/Quadric.cpp
+++ b/Engine/Source/Developer/MeshSimplifier/Private/Quadric.cpp
@@ -4,7 +4,9 @@
 
 DEFINE_LOG_CATEGORY_STATIC( LogQuadric, Log, All );
 
+#ifdef _MSC_VER
 #pragma float_control( precise, on, push )
+#endif
 
 // LUP factorization using Doolittle's method with partial pivoting
 template< typename T >
@@ -1057,4 +1059,6 @@ bool FQuadricAttrOptimizer::OptimizeLinear( const FVector& Position0, const FVec
 	return true;
 }
 
-#pragma float_control( pop )
\ No newline at end of file
+#ifdef _MSC_VER
+#pragma float_control( pop )
+#endif
\ No newline at end of file
diff --git a/Engine/Source/Developer/MeshSimplifier/Private/Quadric.h b/Engine/Source/Developer/MeshSimplifier/Private/Quadric.h
--- a/Engine/Source/Developer/MeshSimplifier/Private/Quadric.h
+++ b/Engine/Source/Developer/MeshSimplifier/Private/Quadric.h
@@ -9,7 +9,9 @@
 // [ Hoppe 2000, "Efficient minimization of new quadric metric for simplifying meshes with appearance attributes" ]
 
 // doubles needed for precision
+#ifdef _MSC_VER
 #pragma float_control( precise, on, push )
+#endif
 
 #define WEIGHT_BY_AREA		1
 #define VOLUME_CONSTRAINT	1
@@ -578,4 +580,6 @@ inline void FQuadricAttrOptimizer::AddQuadric( const FQuadricAttr& RESTRICT q, u
 	}
 }
 
-#pragma float_control( pop )
\ No newline at end of file
+#ifdef _MSC_VER
+#pragma float_control( pop )
+#endif
\ No newline at end of file
