Fix clang build pragma error

diff --git a/Engine/Source/Developer/MeshSimplifier/Private/Quadric.cpp b/Engine/Source/Developer/MeshSimplifier/Private/Quadric.cpp
--- a/Engine/Source/Developer/MeshSimplifier/Private/Quadric.cpp
+++ b/Engine/Source/Developer/MeshSimplifier/Private/Quadric.cpp
@@ -4,7 +4,7 @@
 
 DEFINE_LOG_CATEGORY_STATIC( LogQuadric, Log, All );
 
-#ifdef _MSC_VER
+#if defined(_MSC_VER) && !defined(__clang__)
 #pragma float_control( precise, on, push )
 #endif
 
diff --git a/Engine/Source/Developer/MeshSimplifier/Private/Quadric.h b/Engine/Source/Developer/MeshSimplifier/Private/Quadric.h
--- a/Engine/Source/Developer/MeshSimplifier/Private/Quadric.h
+++ b/Engine/Source/Developer/MeshSimplifier/Private/Quadric.h
@@ -9,7 +9,7 @@
 // [ Hoppe 2000, "Efficient minimization of new quadric metric for simplifying meshes with appearance attributes" ]
 
 // doubles needed for precision
-#ifdef _MSC_VER
+#if defined(_MSC_VER) && !defined(__clang__)
 #pragma float_control( precise, on, push )
 #endif
 
