Fix hashing for 0 == -0

diff --git a/Engine/Source/Developer/MeshSimplifier/Private/MeshSimplify.h b/Engine/Source/Developer/MeshSimplifier/Private/MeshSimplify.h
--- a/Engine/Source/Developer/MeshSimplifier/Private/MeshSimplify.h
+++ b/Engine/Source/Developer/MeshSimplifier/Private/MeshSimplify.h
@@ -24,7 +24,11 @@ FORCEINLINE uint32 HashPosition( const FVector& Position )
 	y.f = Position.Y;
 	z.f = Position.Z;
 
-	return Murmur32( { x.i, y.i, z.i } );
+	return Murmur32( {
+		Position.X == 0.0f ? 0u : x.i,
+		Position.Y == 0.0f ? 0u : y.i,
+		Position.Z == 0.0f ? 0u : z.i
+	} );
 }
 
 FORCEINLINE uint32 Cycle3( uint32 Value )
