Fix stupid static analysis

diff --git a/Engine/Source/Developer/MeshSimplifier/Private/Quadric.cpp b/Engine/Source/Developer/MeshSimplifier/Private/Quadric.cpp
--- a/Engine/Source/Developer/MeshSimplifier/Private/Quadric.cpp
+++ b/Engine/Source/Developer/MeshSimplifier/Private/Quadric.cpp
@@ -6,6 +6,7 @@ DEFINE_LOG_CATEGORY_STATIC( LogQuadric, Log, All );
 
 #if defined(_MSC_VER) && !defined(__clang__)
 #pragma float_control( precise, on, push )
+#pragma warning(disable:6011)
 #endif
 
 // LUP factorization using Doolittle's method with partial pivoting
diff --git a/Engine/Source/Developer/MeshSimplifier/Private/Quadric.h b/Engine/Source/Developer/MeshSimplifier/Private/Quadric.h
--- a/Engine/Source/Developer/MeshSimplifier/Private/Quadric.h
+++ b/Engine/Source/Developer/MeshSimplifier/Private/Quadric.h
@@ -11,6 +11,7 @@
 // doubles needed for precision
 #if defined(_MSC_VER) && !defined(__clang__)
 #pragma float_control( precise, on, push )
+#pragma warning(disable:6011)
 #endif
 
 #define WEIGHT_BY_AREA		1
