Fix for simplifier reducing down to no triangles.

diff --git a/Engine/Source/Developer/MeshSimplifier/Private/QuadricMeshReduction.cpp b/Engine/Source/Developer/MeshSimplifier/Private/QuadricMeshReduction.cpp
--- a/Engine/Source/Developer/MeshSimplifier/Private/QuadricMeshReduction.cpp
+++ b/Engine/Source/Developer/MeshSimplifier/Private/QuadricMeshReduction.cpp
@@ -407,8 +407,13 @@ public:
 
 			float MaxErrorSqr = Simplifier.Simplify( TargetNumVerts, TargetNumTris );
 
-			check( Simplifier.GetRemainingNumVerts() > 0 );
-			check( Simplifier.GetRemainingNumTris() > 0 );
+			if( Simplifier.GetRemainingNumVerts() == 0 || Simplifier.GetRemainingNumTris() == 0 )
+			{
+				// Reduced to nothing so just return the orignial.
+				OutReducedMesh = InMesh;
+				OutMaxDeviation = 0.0f;
+				return;
+			}
 		
 			Simplifier.Compact();
 	
