Allow auto screen size for LODs to be over 1. This is required for high poly meshes where simplifying causes only a small deviation.
Change mesh simplifier settings to match closer to what they used to be to prevent large changes in LOD distances from ue4 to ue5.

diff --git a/Engine/Source/Developer/MeshSimplifier/Private/MeshSimplify.cpp b/Engine/Source/Developer/MeshSimplifier/Private/MeshSimplify.cpp
--- a/Engine/Source/Developer/MeshSimplifier/Private/MeshSimplify.cpp
+++ b/Engine/Source/Developer/MeshSimplifier/Private/MeshSimplify.cpp
@@ -492,8 +492,11 @@ float FMeshSimplifier::EvaluateMerge( const FVector3f& Position0, const FVector3
 
 	Error += EdgeError;
 
-	// Limit error to be no greater than the size of the triangles it could affect.
-	Error = FMath::Min( Error, SurfaceArea );
+	if( bLimitErrorToSurfaceArea )
+	{
+		// Limit error to be no greater than the size of the triangles it could affect.
+		Error = FMath::Min( Error, SurfaceArea );
+	}
 
 	if( bMoveVerts )
 	{
diff --git a/Engine/Source/Developer/MeshSimplifier/Private/MeshSimplify.h b/Engine/Source/Developer/MeshSimplifier/Private/MeshSimplify.h
--- a/Engine/Source/Developer/MeshSimplifier/Private/MeshSimplify.h
+++ b/Engine/Source/Developer/MeshSimplifier/Private/MeshSimplify.h
@@ -53,6 +53,7 @@ public:
 	void		SetAttributeWeights( const float* Weights )			{ AttributeWeights = Weights; }
 	void		SetEdgeWeight( float Weight )						{ EdgeWeight = Weight; }
 	void		SetCorrectAttributes( void (*Function)( float* ) )	{ CorrectAttributes = Function; }
+	void		SetLimitErrorToSurfaceArea( bool Value )			{ bLimitErrorToSurfaceArea = Value; }
 
 	QUADRICMESHREDUCTION_API void	LockPosition( const FVector3f& Position );
 
@@ -64,12 +65,12 @@ public:
 	uint32		GetRemainingNumVerts() const	{ return RemainingNumVerts; }
 	uint32		GetRemainingNumTris() const		{ return RemainingNumTris; }
 
-protected:
-	const int32 DegreeLimit			= 24;
-	const float DegreePenalty		= 0.5f;
-	const float LockPenalty			= 1e8f;
-	const float InversionPenalty	= 100.0f;
+	int32 DegreeLimit		= 24;
+	float DegreePenalty		= 0.5f;
+	float LockPenalty		= 1e8f;
+	float InversionPenalty	= 100.0f;
 
+protected:
 	uint32		NumVerts;
 	uint32		NumIndexes;
 	uint32		NumAttributes;
@@ -85,6 +86,7 @@ protected:
 	const float*	AttributeWeights = nullptr;
 	float			EdgeWeight = 8.0f;
 	void			(*CorrectAttributes)( float* ) = nullptr;
+	bool			bLimitErrorToSurfaceArea = true;
 
 	FHashTable		VertHash;
 	FHashTable		CornerHash;
diff --git a/Engine/Source/Developer/MeshSimplifier/Private/QuadricMeshReduction.cpp b/Engine/Source/Developer/MeshSimplifier/Private/QuadricMeshReduction.cpp
--- a/Engine/Source/Developer/MeshSimplifier/Private/QuadricMeshReduction.cpp
+++ b/Engine/Source/Developer/MeshSimplifier/Private/QuadricMeshReduction.cpp
@@ -39,7 +39,7 @@ IMPLEMENT_MODULE(FQuadricSimplifierMeshReductionModule, QuadricMeshReduction);
 
 void CorrectAttributes( float* Attributes )
 {
-	FVector3f& Normal		= *reinterpret_cast< FVector3f* >( Attributes );
+	FVector3f& Normal	= *reinterpret_cast< FVector3f* >( Attributes );
 	FVector3f& TangentX	= *reinterpret_cast< FVector3f* >( Attributes + 3 );
 	FVector3f& TangentY	= *reinterpret_cast< FVector3f* >( Attributes + 3 + 3 );
 	FLinearColor& Color	= *reinterpret_cast< FLinearColor* >( Attributes + 3 + 3 + 3 );
@@ -64,7 +64,7 @@ public:
 		// VersionString.ParseIntoArray(SplitVersionString, TEXT("_"), true);
 		// bool bUseQuadricSimplier = SplitVersionString[0].Equals("QuadricMeshReduction");
 
-		static FString Version = TEXT("QuadricMeshReduction_V2.0");
+		static FString Version = TEXT("QuadricMeshReduction_V2.1");
 		static FString AltVersion = TEXT("QuadricMeshReduction_V2.0_OldSimplifier");
 		return bUseOldMeshSimplifier ? AltVersion : Version;
 	}
@@ -346,24 +346,12 @@ public:
 			{
 				using VertType = TVertSimp< NumTexCoords >;
 
-				float TriangleSize = FMath::Sqrt( SurfaceArea / NumTris );
-	
-				FFloat32 CurrentSize( FMath::Max( TriangleSize, THRESH_POINTS_ARE_SAME ) );
-				FFloat32 DesiredSize( 0.25f );
-				FFloat32 Scale( 1.0f );
-
-				// Lossless scaling by only changing the float exponent.
-				int32 Exponent = FMath::Clamp( (int)DesiredSize.Components.Exponent - (int)CurrentSize.Components.Exponent, -126, 127 );
-				Scale.Components.Exponent = Exponent + 127;	//ExpBias
-				float PositionScale = Scale.FloatValue;
-
-
 				const uint32 NumAttributes = ( sizeof( VertType ) - sizeof( FVector3f ) ) / sizeof(float);
 				float AttributeWeights[ NumAttributes ] =
 				{
-					1.0f, 1.0f, 1.0f,		// Normal
-					0.006f, 0.006f, 0.006f,	// Tangent[0]
-					0.006f, 0.006f, 0.006f	// Tangent[1]
+					16.0f, 16.0f, 16.0f,// Normal
+					0.1f, 0.1f, 0.1f,	// Tangent[0]
+					0.1f, 0.1f, 0.1f	// Tangent[1]
 				};
 				float* ColorWeights = AttributeWeights + 9;
 				float* UVWeights = ColorWeights + 4;
@@ -373,13 +361,13 @@ public:
 				// Set weights if they are used
 				if( bHasColors )
 				{
-					ColorWeights[0] = 0.0625f;
-					ColorWeights[1] = 0.0625f;
-					ColorWeights[2] = 0.0625f;
-					ColorWeights[3] = 0.0625f;
+					ColorWeights[0] = 0.1f;
+					ColorWeights[1] = 0.1f;
+					ColorWeights[2] = 0.1f;
+					ColorWeights[3] = 0.1f;
 				}
 
-				float UVWeight = 1.0f / ( 32.0f * InVertexUVs.GetNumChannels() );
+				float UVWeight = 0.5f;
 				for( int32 UVIndex = 0; UVIndex < InVertexUVs.GetNumChannels(); UVIndex++ )
 				{
 					// Normalize UVWeights using min/max UV range.
@@ -399,18 +387,17 @@ public:
 					UVWeights[ 2 * UVIndex + 1 ] = UVWeight / FMath::Max( 1.0f, MaxUV - MinUV );
 				}
 
-				for( auto& Vert : Verts )
-				{
-					Vert.Position *= PositionScale;
-				}
-
 				FMeshSimplifier Simplifier( (float*)Verts.GetData(), Verts.Num(), Indexes.GetData(), Indexes.Num(), MaterialIndexes.GetData(), NumAttributes );
 
 				Simplifier.SetAttributeWeights( AttributeWeights );
 				Simplifier.SetCorrectAttributes( CorrectAttributes );
-				Simplifier.SetEdgeWeight( 4.0f );
+				Simplifier.SetEdgeWeight( 512.0f );
+				Simplifier.SetLimitErrorToSurfaceArea( false );
+
+				Simplifier.DegreePenalty = 100.0f;
+				Simplifier.InversionPenalty = 1000000.0f;
 
-				float MaxErrorSqr = Simplifier.Simplify( TargetNumVerts, TargetNumTris, 0.0f, 0, 0, MAX_flt );
+				float MaxErrorSqr = Simplifier.Simplify( TargetNumVerts, TargetNumTris, 0.0f, 4, 2, MAX_flt );
 
 				if( Simplifier.GetRemainingNumVerts() == 0 || Simplifier.GetRemainingNumTris() == 0 )
 				{
@@ -430,13 +417,7 @@ public:
 				NumTris = Simplifier.GetRemainingNumTris();
 				NumIndexes = NumTris * 3;
 
-				float InvScale = 1.0f / PositionScale;
-				for( auto& Vert : Verts )
-				{
-					Vert.Position *= InvScale;
-				}
-
-				OutMaxDeviation = FMath::Sqrt( MaxErrorSqr ) * InvScale;
+				OutMaxDeviation = FMath::Sqrt( MaxErrorSqr ) / 8.0f;
 			}
 			else
 			{
diff --git a/Engine/Source/Runtime/Engine/Private/StaticMesh.cpp b/Engine/Source/Runtime/Engine/Private/StaticMesh.cpp
--- a/Engine/Source/Runtime/Engine/Private/StaticMesh.cpp
+++ b/Engine/Source/Runtime/Engine/Private/StaticMesh.cpp
@@ -1952,7 +1952,7 @@ void FStaticMeshRenderData::ResolveSectionInfo(UStaticMesh* Owner)
 		{
 			if (LODIndex == 0)
 			{
-				ScreenSize[LODIndex].Default = 1.0f;
+				ScreenSize[LODIndex].Default = 2.0f;
 			}
 			else if(LOD.MaxDeviation <= 0.0f)
 			{
