Add spring quadrics

diff --git a/Engine/Source/Developer/MeshSimplifier/Private/Quadric.cpp b/Engine/Source/Developer/MeshSimplifier/Private/Quadric.cpp
--- a/Engine/Source/Developer/MeshSimplifier/Private/Quadric.cpp
+++ b/Engine/Source/Developer/MeshSimplifier/Private/Quadric.cpp
@@ -520,6 +520,45 @@ FQuadric::FQuadric( const QVec3 p0, const QVec3 p1, const QVec3 p2 )
 #endif
 }
 
+FQuadric::FQuadric( const QVec3 p )
+{
+	// (v - p)^T (v - p)
+	// v^T I v - 2 p^T v + p^T p
+	nxx = 1.0;
+	nyy = 1.0;
+	nzz = 1.0;
+
+	nxy = 0.0;
+	nxz = 0.0;
+	nyz = 0.0;
+
+	dn = -p;
+	d2 = p | p;
+
+	a = 0.0;
+}
+
+FQuadric::FQuadric( const QVec3 n, const QVec3 p )
+{
+	// nn^T = projection matrix
+	//( v - nn^T v )^T ( v - nn^T v )
+	// v^T ( I - nn^T ) v - 2p^T ( I - nn^T ) v + (p^T p - p^T nn^T p)
+	nxx = 1.0 - n.x * n.x;
+	nyy = 1.0 - n.y * n.y;
+	nzz = 1.0 - n.z * n.z;
+
+	nxy = -n.x * n.y;
+	nxz = -n.x * n.z;
+	nyz = -n.y * n.z;
+
+	const QScalar dist = -( n | p );
+
+	dn = -p - dist * n;
+	d2 = (p | p) - dist * dist;
+
+	a = 0.0;
+}
+
 float FQuadric::Evaluate( const FVector3f& Point ) const
 {
 	// Q(v) = vt*A*v + 2*bt*v + c
diff --git a/Engine/Source/Developer/MeshSimplifier/Private/Quadric.h b/Engine/Source/Developer/MeshSimplifier/Private/Quadric.h
--- a/Engine/Source/Developer/MeshSimplifier/Private/Quadric.h
+++ b/Engine/Source/Developer/MeshSimplifier/Private/Quadric.h
@@ -67,6 +67,11 @@ struct TVec3
 		: x( v.X ), y( v.Y ), z( v.Z )
 	{}
 
+	TVec3 operator-() const
+	{
+		return TVec3<T>( -x, -y, -z );
+	}
+
 	TVec3 operator*( T Scalar ) const
 	{
 		return TVec3<T>(
@@ -227,6 +232,11 @@ class FQuadric
 {
 public:
 				FQuadric() {}
+				// Distance from point
+				FQuadric( const QVec3 p );
+				// Distance from line, n must be normalized
+				FQuadric( const QVec3 n, const QVec3 p );
+				// Distance from triangle
 				FQuadric( const QVec3 p0, const QVec3 p1, const QVec3 p2 );
 	
 	void		Zero();
