Re-enable the new static mesh simplifier

#tests a/b tested on multiple platforms to make sure no perf or memory regression; sanity checked several static meshes to make sure auto calculated screensizes stayed similar

diff --git a/Engine/Source/Developer/MeshSimplifier/Private/QuadricMeshReduction.cpp b/Engine/Source/Developer/MeshSimplifier/Private/QuadricMeshReduction.cpp
--- a/Engine/Source/Developer/MeshSimplifier/Private/QuadricMeshReduction.cpp
+++ b/Engine/Source/Developer/MeshSimplifier/Private/QuadricMeshReduction.cpp
@@ -341,6 +341,8 @@ public:
 
 			// We need a minimum of 2 triangles, to see the object on both side. If we use one, we will end up with zero triangle when we will remove a shared edge
 			TargetNumTris = FMath::Max( TargetNumTris, 2u );
+			// Clamp to a minimum of 4 vertices, ReductionSettings.PercentVertices can be zero which makes TargetNumVerts also zero
+			TargetNumVerts = FMath::Max(TargetNumVerts, 4u);
 
 			if( TargetNumVerts < NumVerts || TargetNumTris < NumTris )
 			{
