Fix for overflow error when simplifying large meshes.

diff --git a/Engine/Source/Developer/MeshSimplifier/Private/MeshSimplify.cpp b/Engine/Source/Developer/MeshSimplifier/Private/MeshSimplify.cpp
--- a/Engine/Source/Developer/MeshSimplifier/Private/MeshSimplify.cpp
+++ b/Engine/Source/Developer/MeshSimplifier/Private/MeshSimplify.cpp
@@ -282,7 +282,7 @@ float FMeshSimplifier::EvaluateMerge( const FVector3f& Position0, const FVector3
 	TArray< uint32, TInlineAllocator<8> >	WedgeIDs;
 	TArray< uint8, TInlineAllocator<1024> >	WedgeQuadrics;
 	
-	const uint32 QuadricSize = sizeof( FQuadricAttr ) + NumAttributes * 4 * sizeof( QScalar );
+	const SIZE_T QuadricSize = sizeof( FQuadricAttr ) + NumAttributes * 4 * sizeof( QScalar );
 
 	auto GetWedgeQuadric =
 		[ &WedgeQuadrics, QuadricSize ]( int32 WedgeIndex ) -> FQuadricAttr&
@@ -852,7 +852,7 @@ float FMeshSimplifier::Simplify(
 	check( TargetNumTris >= LimitNumTris );
 	check( TargetError <= LimitError );
 
-	const uint32 QuadricSize = sizeof( FQuadricAttr ) + NumAttributes * 4 * sizeof( QScalar );
+	const SIZE_T QuadricSize = sizeof( FQuadricAttr ) + NumAttributes * 4 * sizeof( QScalar );
 
 	TriQuadrics.AddUninitialized( NumTris * QuadricSize );
 	for( uint32 TriIndex = 0; TriIndex < NumTris; TriIndex++ )
diff --git a/Engine/Source/Developer/MeshSimplifier/Private/MeshSimplify.h b/Engine/Source/Developer/MeshSimplifier/Private/MeshSimplify.h
--- a/Engine/Source/Developer/MeshSimplifier/Private/MeshSimplify.h
+++ b/Engine/Source/Developer/MeshSimplifier/Private/MeshSimplify.h
@@ -178,7 +178,7 @@ FORCEINLINE float* FMeshSimplifier::GetAttributes( uint32 VertIndex )
 
 FORCEINLINE FQuadricAttr& FMeshSimplifier::GetTriQuadric( uint32 TriIndex )
 {
-	const uint32 QuadricSize = sizeof( FQuadricAttr ) + NumAttributes * 4 * sizeof( QScalar );
+	const SIZE_T QuadricSize = sizeof( FQuadricAttr ) + NumAttributes * 4 * sizeof( QScalar );
 	return *reinterpret_cast< FQuadricAttr* >( &TriQuadrics[ TriIndex * QuadricSize ] );
 }
 
