Fixed compile errors in MeshSimplify.cpp
-Linux doesn't compile because fopen_s is not a standard function
-Missing include for Math/RandomStream.h

diff --git a/Engine/Source/Developer/MeshSimplifier/Private/MeshSimplify.cpp b/Engine/Source/Developer/MeshSimplifier/Private/MeshSimplify.cpp
--- a/Engine/Source/Developer/MeshSimplifier/Private/MeshSimplify.cpp
+++ b/Engine/Source/Developer/MeshSimplifier/Private/MeshSimplify.cpp
@@ -1,6 +1,7 @@
 // Copyright Epic Games, Inc. All Rights Reserved.
 
 #include "MeshSimplify.h"
+#include "Math/RandomStream.h"
 
 FMeshSimplifier::FMeshSimplifier( float* InVerts, uint32 InNumVerts, uint32* InIndexes, uint32 InNumIndexes, int32* InMaterialIndexes, uint32 InNumAttributes )
 	: NumVerts( InNumVerts )
@@ -1089,8 +1090,12 @@ void FMeshSimplifier::PreserveSurfaceArea()
 
 void FMeshSimplifier::DumpOBJ( const char* Filename )
 {
+#if PLATFORM_WINDOWS
 	FILE* File = nullptr;
-	fopen_s( &File, Filename, "wb" );
+	fopen_s(&File, Filename, "wb");
+#else
+	FILE* File = fopen(Filename, "wb");
+#endif
 
 	for( uint32 i = 0; PairHeap.Num() > 0; i++ )
 	{
