Added nullptr check for File in FMeshSimplifier::DumpOBJ to make PVS studio happy (Issue: 219027)

diff --git a/Engine/Source/Developer/MeshSimplifier/Private/MeshSimplify.cpp b/Engine/Source/Developer/MeshSimplifier/Private/MeshSimplify.cpp
--- a/Engine/Source/Developer/MeshSimplifier/Private/MeshSimplify.cpp
+++ b/Engine/Source/Developer/MeshSimplifier/Private/MeshSimplify.cpp
@@ -1096,27 +1096,30 @@ void FMeshSimplifier::DumpOBJ( const char* Filename )
 #else
 	FILE* File = fopen(Filename, "wb");
 #endif
-
-	for( uint32 i = 0; PairHeap.Num() > 0; i++ )
+	
+	if( File )
 	{
-		uint32 PairIndex = PairHeap.Top();
+		for( uint32 i = 0; PairHeap.Num() > 0; i++ )
+		{
+			uint32 PairIndex = PairHeap.Top();
 
-		float Error = PairHeap.GetKey( PairIndex );
-		PairHeap.Pop();
+			float Error = PairHeap.GetKey( PairIndex );
+			PairHeap.Pop();
 
-		const FPair& Pair = Pairs[ PairIndex ];
+			const FPair& Pair = Pairs[ PairIndex ];
 
-		const FVector3f& p0 = Pair.Position0;
-		const FVector3f& p1 = Pair.Position1;
+			const FVector3f& p0 = Pair.Position0;
+			const FVector3f& p1 = Pair.Position1;
 
-		fprintf( File, "v %f %f %f\n", p0.X, p0.Y, p0.Z );
-		fprintf( File, "v %f %f %f\n", p1.X, p1.Y, p1.Z );
+			fprintf( File, "v %f %f %f\n", p0.X, p0.Y, p0.Z );
+			fprintf( File, "v %f %f %f\n", p1.X, p1.Y, p1.Z );
 
-		// +1 as Wavefront files are 1 index based
-		fprintf( File, "l %u %u # %u %f\n", i*2 + 1, i*2 + 2, i + 1, Error );
-	}
+			// +1 as Wavefront files are 1 index based
+			fprintf( File, "l %u %u # %u %f\n", i*2 + 1, i*2 + 2, i + 1, Error );
+		}
 	
-	fclose( File );
+		fclose( File );
+	}
 }
 
 void FMeshSimplifier::Compact()
