Fixed bug where single edge collapse removes all remaining triangles.

Deleted old simplifier

diff --git a/Engine/Source/Developer/MeshSimplifier/Private/Cache.h b/Engine/Source/Developer/MeshSimplifier/Private/Cache.h
deleted file mode 100644
diff --git a/Engine/Source/Developer/MeshSimplifier/Private/MeshSimplify.cpp b/Engine/Source/Developer/MeshSimplifier/Private/MeshSimplify.cpp
--- a/Engine/Source/Developer/MeshSimplifier/Private/MeshSimplify.cpp
+++ b/Engine/Source/Developer/MeshSimplifier/Private/MeshSimplify.cpp
@@ -175,14 +175,7 @@ void FMeshSimplifier::CalcEdgeQuadric( uint32 EdgeIndex )
 		}
 	}
 
-	const QVec3 p0 = GetPosition( Indexes[ TriIndex * 3 + 0 ] );
-	const QVec3 p1 = GetPosition( Indexes[ TriIndex * 3 + 1 ] );
-	const QVec3 p2 = GetPosition( Indexes[ TriIndex * 3 + 2 ] );
-
-	const QVec3 Normal = ( p2 - p0 ) ^ ( p1 - p0 );
-
 	// Didn't find matching edge. Add edge constraint.
-	//EdgeQuadrics[ EdgeIndex ] = FEdgeQuadric( GetPosition( VertIndex0 ), GetPosition( VertIndex1 ), Normal, Weight );
 	EdgeQuadrics[ EdgeIndex ] = FEdgeQuadric( GetPosition( VertIndex0 ), GetPosition( VertIndex1 ), Weight );
 	EdgeQuadricsValid[ EdgeIndex ] = true;
 }
@@ -751,6 +744,27 @@ bool FMeshSimplifier::TriWillInvert( uint32 TriIndex, const FVector3f& NewPositi
 	return false;
 }
 
+void FMeshSimplifier::RemoveTri( uint32 TriIndex )
+{
+	check( !TriRemoved[ TriIndex ] );
+
+	TriRemoved[ TriIndex ] = true;
+	RemainingNumTris--;
+
+	// Remove references to tri
+	for( uint32 k = 0; k < 3; k++ )
+	{
+		uint32 Corner = TriIndex * 3 + k;
+		uint32 VertIndex = Indexes[ Corner ];
+		uint32 Hash = HashPosition( GetPosition( VertIndex ) );
+
+		CornerHash.Remove( Hash, Corner );
+		EdgeQuadricsValid[ Corner ] = false;
+
+		SetVertIndex( Corner, ~0u );
+	}
+}
+
 void FMeshSimplifier::FixUpTri( uint32 TriIndex )
 {
 	check( !TriRemoved[ TriIndex ] );
@@ -781,27 +795,9 @@ void FMeshSimplifier::FixUpTri( uint32 TriIndex )
 	}
 
 	if( bRemoveTri )
-	{
-		TriRemoved[ TriIndex ] = true;
-		RemainingNumTris--;
-
-		// Remove references to tri
-		for( uint32 k = 0; k < 3; k++ )
-		{
-			uint32 Corner = TriIndex * 3 + k;
-			uint32 VertIndex = Indexes[ Corner ];
-			uint32 Hash = HashPosition( GetPosition( VertIndex ) );
-
-			CornerHash.Remove( Hash, Corner );
-			EdgeQuadricsValid[ Corner ] = false;
-
-			SetVertIndex( Corner, ~0u );
-		}
-	}
+		RemoveTri( TriIndex );
 	else
-	{
 		CalcTriQuadric( TriIndex );
-	}
 }
 
 bool FMeshSimplifier::IsDuplicateTri( uint32 TriIndex ) const
@@ -958,6 +954,32 @@ float FMeshSimplifier::Simplify(
 		}
 		ReevaluatePairs.Reset();
 	}
+
+	// If couldn't hit targets through regular edge collapses, resort to randomly removing triangles.
+	{
+		uint32 TriIndex = 0;
+		while(1)
+		{
+			if( RemainingNumVerts	<= TargetNumVerts &&
+				RemainingNumTris	<= TargetNumTris &&
+				MaxError			>= TargetError )
+			{
+				break;
+			}
+	
+			if( RemainingNumVerts	<= LimitNumVerts ||
+				RemainingNumTris	<= LimitNumTris ||
+				MaxError			>= LimitError )
+			{
+				break;
+			}
+
+			while( TriRemoved[ TriIndex ] )
+				TriIndex++;
+	
+			RemoveTri( TriIndex );
+		}
+	}
 	
 	return MaxError;
 }
diff --git a/Engine/Source/Developer/MeshSimplifier/Private/MeshSimplify.h b/Engine/Source/Developer/MeshSimplifier/Private/MeshSimplify.h
--- a/Engine/Source/Developer/MeshSimplifier/Private/MeshSimplify.h
+++ b/Engine/Source/Developer/MeshSimplifier/Private/MeshSimplify.h
@@ -5,10 +5,8 @@
 
 #include "CoreMinimal.h"
 
-#define SIMP_NEW_MERGE	1
 #define SIMP_REBASE		1
 
-#include "MeshSimplifyElements.h"
 #include "Quadric.h"
 #include "Containers/HashTable.h"
 #include "Containers/BinaryHeap.h"
@@ -164,6 +162,7 @@ protected:
 	uint32	CornerIndexMoved( uint32 TriIndex ) const;
 	bool	TriWillInvert( uint32 TriIndex, const FVector3f& NewPosition ) const;
 
+	void	RemoveTri( uint32 TriIndex );
 	void	FixUpTri( uint32 TriIndex );
 	bool	IsDuplicateTri( uint32 TriIndex ) const;
 	void	SetVertIndex( uint32 Corner, uint32 NewVertIndex );
@@ -237,1919 +236,4 @@ void FMeshSimplifier::ForAllPairs( const FVector3f& Position, FuncType&& Functio
 			Function( PairIndex );
 		}
 	}
-}
-
-
-
-template< typename T, uint32 NumAttributes >
-class TMeshSimplifier
-{
-	typedef typename TSimpVert<T>::TriIterator	TriIterator;
-	typedef TQuadricAttr< NumAttributes >		QuadricType;
-
-public:
-						TMeshSimplifier( const T* Verts, uint32 NumVerts, const uint32* Indexes, uint32 NumIndexes );
-						~TMeshSimplifier();
-
-	void				SetAttributeWeights( const float* weights );
-	void				SetEdgeWeight( float Weight );
-	void				SetBoundaryLocked();
-
-	void				InitCosts();
-
-	float				SimplifyMesh( float maxErrorLimit, int minTris, int minVerts = 0);
-
-	int					GetNumVerts() const { return numVerts; }
-	int					GetNumTris() const { return numTris; }
-
-	void				OutputMesh( T* Verts, uint32* Indexes, uint32* OutNumVertices = nullptr, uint32* OutNumIndices = nullptr ) const;
-
-	template< typename FaceType >
-	void				CompactFaceData( TArray< FaceType>& FaceData );
-
-protected:
-	void				LockVertFlags( uint32 flag );
-	void				UnlockVertFlags( uint32 flag );
-
-	void				LockTriFlags( uint32 flag );
-	void				UnlockTriFlags( uint32 flag );
-
-	void				GatherUpdates( TSimpVert<T>* v );
-
-	void				GroupVerts();
-	void				GroupEdges();
-
-	void				InitVert( TSimpVert<T>* v );
-
-	QuadricType&		GetQuadric( const TSimpTri<T>* Tri );
-	QuadricType&		GetQuadric( TSimpVert<T>* v );
-	FQuadric&			GetEdgeQuadric( TSimpVert<T>* v );
-
-	// TODO move away from pointers and remove these functions
-	uint32				GetVertIndex( const TSimpVert<T>* vert ) const;
-	uint32				GetTriIndex( const TSimpTri<T>* tri ) const;
-	uint32				GetEdgeIndex( const TSimpEdge<T>* edge ) const;
-
-	uint32				HashPoint( const FVector3f& p ) const;
-	uint32				HashEdge( const TSimpVert<T>* u, const TSimpVert<T>* v ) const;
-	TSimpEdge<T>*		FindEdge( const TSimpVert<T>* u, const TSimpVert<T>* v );
-	
-	void				RemoveEdge( TSimpEdge<T>* edge );
-	void				RemoveVert( TSimpVert<T>* vert );
-	void				ReplaceEdgeVert( const TSimpVert<T>* oldV, const TSimpVert<T>* otherV, TSimpVert<T>* newV );
-	void				CollapseEdgeVert( const TSimpVert<T>* oldV, const TSimpVert<T>* otherV, TSimpVert<T>* newV );
-
-	float				Merge( TSimpVert<T>* v0, TSimpVert<T>* v1, bool bMoveVerts );
-	float				ComputeMergePenalty( TSimpVert<T>* v0, TSimpVert<T>* v1, const FVector3f& NewPos );
-#if !SIMP_NEW_MERGE
-	float				ComputeNewVerts( TSimpEdge<T>* edge, TArray< T, TInlineAllocator<16> >& newVerts );
-	float				ComputeEdgeCollapseCost( TSimpEdge<T>* edge, bool bMoveVerts );
-#endif
-	void				AddEdgeToHeap( TSimpEdge<T>* edge );
-	void				Collapse( TSimpEdge<T>* edge );
-
-	void				UpdateTris();
-	void				UpdateVerts();
-	void				UpdateEdges();
-
-	uint32				vertFlagLock;
-	uint32				triFlagLock;
-	
-	float				attributeWeights[ NumAttributes ];
-	float				EdgeWeight = 16.0f;
-	
-	TSimpVert<T>*		sVerts;
-	TSimpTri<T>*		sTris;
-
-	int					numSVerts;
-	int					numSTris;
-
-	int					numVerts;
-	int					numTris;
-
-	TArray< TSimpEdge<T> >	edges;
-	FHashTable				edgeHash;
-	FBinaryHeap<float>		edgeHeap;
-
-	TArray< TSimpVert<T>* >	updateVerts;
-	TArray< TSimpTri<T>* >	updateTris;
-	TArray< TSimpEdge<T>* >	updateEdges;
-
-	TBitArray<>				VertQuadricsValid;
-	TArray< QuadricType >	VertQuadrics;
-
-	TBitArray<>				TriQuadricsValid;
-	TArray< QuadricType >	TriQuadrics;
-
-	TBitArray<>				EdgeQuadricsValid;
-	TArray< FQuadric >		EdgeQuadrics;
-
-	FHashTable			AdjCorners;
-	FDisjointSet		WedgeSet;
-};
-
-
-//=============
-// TMeshSimplifier
-//=============
-template< typename T, uint32 NumAttributes >
-TMeshSimplifier<T, NumAttributes>::TMeshSimplifier( const T* Verts, uint32 NumVerts, const uint32* Indexes, uint32 NumIndexes )
-	: edgeHash( 1 << FMath::Min( 16u, FMath::FloorLog2( NumVerts ) ) )
-	, AdjCorners( 64 )
-{
-	TRACE_CPUPROFILER_EVENT_SCOPE(TMeshSimplifier::TMeshSimplifier);
-	vertFlagLock = 0;
-	triFlagLock = 0;
-
-	for( uint32 i = 0; i < NumAttributes; i++ )
-	{
-		attributeWeights[i] = 1.0f;
-	}
-
-	numSVerts = NumVerts;
-	numSTris = NumIndexes / 3;
-
-	this->numVerts = numSVerts;
-	this->numTris = numSTris;
-
-	sVerts = new TSimpVert<T>[ numSVerts ];
-	sTris = new TSimpTri<T>[ numSTris ];
-
-	VertQuadricsValid.Init( false, numSVerts );
-	VertQuadrics.SetNum( numSVerts );
-
-	TriQuadricsValid.Init( false, numSTris );
-	TriQuadrics.SetNum( numSTris );
-
-	EdgeQuadricsValid.Init( false, numSVerts );
-	EdgeQuadrics.SetNum( numSVerts );
-
-	for( int i = 0; i < numSVerts; i++ )
-	{
-		sVerts[i].vert = Verts[i];
-	}
-	
-	for( int i = 0; i < numSTris; i++ )
-	{
-		for( int j = 0; j < 3; j++ )
-		{
-			sTris[i].verts[j] = &sVerts[ Indexes[3*i+j] ];
-			sTris[i].verts[j]->adjTris.Add( &sTris[i] );
-		}
-	}
-
-	GroupVerts();
-
-	int maxEdgeSize = FMath::Min( 3 * numSTris, 3 * numSVerts - 6 );
-	edges.Empty( maxEdgeSize );
-	for( int i = 0; i < numSVerts; i++ )
-	{
-		InitVert( &sVerts[i] );
-	}
-
-	// Guessed wrong on num edges. Array was resized so fix up pointers.
-	if( edges.Num() > maxEdgeSize )
-	{
-		for( int i = 0; i < edges.Num(); i++ )
-		{
-			TSimpEdge<T>& edge = edges[i];
-			edge.next = &edge;
-			edge.prev = &edge;
-		}
-	}
-
-	GroupEdges();
-
-	edgeHash.Resize( edges.Num() );
-	for( int i = 0; i < edges.Num(); i++ )
-	{
-		edgeHash.Add( HashEdge( edges[i].v0, edges[i].v1 ), i );
-	}
-
-	edgeHeap.Resize( edges.Num(), edges.Num() );
-}
-
-template< typename T, uint32 NumAttributes >
-TMeshSimplifier<T, NumAttributes>::~TMeshSimplifier()
-{
-	delete[] sVerts;
-	delete[] sTris;
-}
-
-template< typename T, uint32 NumAttributes >
-void TMeshSimplifier<T, NumAttributes>::SetAttributeWeights( const float* weights )
-{
-	for( uint32 i = 0; i < NumAttributes; i++ )
-	{
-		attributeWeights[i] = weights[i];
-	}
-}
-
-template< typename T, uint32 NumAttributes >
-void TMeshSimplifier<T, NumAttributes>::SetEdgeWeight( float Weight )
-{
-	EdgeWeight = Weight;
-}
-
-template< typename T, uint32 NumAttributes >
-void TMeshSimplifier<T, NumAttributes>::SetBoundaryLocked()
-{
-#if 0
-	TArray< TSimpVert<T>*, TInlineAllocator<64> > adjVerts;
-
-	for( int i = 0; i < numSVerts; i++ )
-	{
-		TSimpVert<T>* v0 = &sVerts[i];
-		check( v0->adjTris.Num() > 0 );
-
-		adjVerts.Reset();
-		v0->FindAdjacentVertsGroup( adjVerts );
-
-		for( TSimpVert<T>* v1 : adjVerts )
-		{
-			if( v0 < v1 )
-			{
-				LockTriFlags( SIMP_MARK1 );
-
-				// set if this edge is boundary
-				// find faces that share v0 and v1
-				v0->EnableAdjTriFlagsGroup( SIMP_MARK1 );
-				v1->DisableAdjTriFlagsGroup( SIMP_MARK1 );
-
-				int faceCount = 0;
-				TSimpVert<T>* vert = v0;
-				do
-				{
-					for( TriIterator j = vert->adjTris.Begin(); j != vert->adjTris.End(); ++j )
-					{
-						TSimpTri<T>* tri = *j;
-						faceCount += tri->TestFlags( SIMP_MARK1 ) ? 0 : 1;
-					}
-					vert = vert->next;
-				} while( vert != v0 );
-
-				v0->DisableAdjTriFlagsGroup( SIMP_MARK1 );
-
-				if( faceCount == 1 )
-				{
-					// only one face on this edge
-					v0->EnableFlagsGroup( SIMP_LOCKED );
-					v1->EnableFlagsGroup( SIMP_LOCKED );
-				}
-
-				UnlockTriFlags( SIMP_MARK1 );
-			}
-		}
-	}
-#else
-	for( auto& edge : edges )
-	{
-		LockTriFlags( SIMP_MARK1 );
-
-		// set if this edge is boundary
-		// find faces that share v0 and v1
-		edge.v0->EnableAdjTriFlagsGroup( SIMP_MARK1 );
-
-		int faceCount = 0;
-		TSimpVert<T>* vert = edge.v1;
-		do
-		{
-			for( TriIterator j = vert->adjTris.Begin(); j != vert->adjTris.End(); ++j )
-			{
-				TSimpTri<T>* tri = *j;
-				faceCount += tri->TestFlags( SIMP_MARK1 ) ? 1 : 0;
-			}
-			vert = vert->next;
-		} while( vert != edge.v1 );
-
-		edge.v0->DisableAdjTriFlagsGroup( SIMP_MARK1 );
-
-		if( faceCount == 1 )
-		{
-			// only one face on this edge
-			edge.v0->EnableFlagsGroup( SIMP_LOCKED );
-			edge.v1->EnableFlagsGroup( SIMP_LOCKED );
-		}
-
-		UnlockTriFlags( SIMP_MARK1 );
-	}
-#endif
-}
-
-// locking functions for nesting safety
-template< typename T, uint32 NumAttributes >
-FORCEINLINE void TMeshSimplifier<T, NumAttributes>::LockVertFlags( uint32 f )
-{
-	checkSlow( ( vertFlagLock & f ) == 0 );
-	vertFlagLock |= f;
-}
-
-template< typename T, uint32 NumAttributes >
-FORCEINLINE void TMeshSimplifier<T, NumAttributes>::UnlockVertFlags( uint32 f )
-{
-	vertFlagLock &= ~f;
-}
-
-template< typename T, uint32 NumAttributes >
-FORCEINLINE void TMeshSimplifier<T, NumAttributes>::LockTriFlags( uint32 f )
-{
-	checkSlow( ( triFlagLock & f ) == 0 );
-	triFlagLock |= f;
-}
-
-template< typename T, uint32 NumAttributes >
-FORCEINLINE void TMeshSimplifier<T, NumAttributes>::UnlockTriFlags( uint32 f )
-{
-	triFlagLock &= ~f;
-}
-
-template< typename T, uint32 NumAttributes >
-void TMeshSimplifier<T, NumAttributes>::InitVert( TSimpVert<T>* v )
-{
-	check( v->adjTris.Num() > 0 );
-
-	TArray< TSimpVert<T>*, TInlineAllocator<64> > adjVerts;
-	v->FindAdjacentVerts( adjVerts );
-
-	TSimpVert<T>* v0 = v;
-	for( TSimpVert<T>* v1 : adjVerts )
-	{
-		if( v0 < v1 )
-		{
-			// add edge
-			edges.AddDefaulted();
-			TSimpEdge<T>& edge = edges.Last();
-			edge.v0 = v0;
-			edge.v1 = v1;
-		}
-	}
-}
-
-template< typename T, uint32 NumAttributes >
-void TMeshSimplifier<T, NumAttributes>::GroupVerts()
-{
-	// group verts that share a point
-	FHashTable HashTable( 1 << FMath::Min( 16u, FMath::FloorLog2( numSVerts / 2 ) ), numSVerts );
-
-	for( int i = 0; i < numSVerts; i++ )
-	{
-		HashTable.Add( HashPoint( sVerts[i].GetPos() ), i );
-	}
-	
-	for( int i = 0; i < numSVerts; i++ )
-	{
-		// already grouped
-		if( sVerts[i].next != &sVerts[i] )
-		{
-			continue;
-		}
-
-		// find any matching verts
-		uint32 hash = HashPoint( sVerts[i].GetPos() );
-		for( int j = HashTable.First( hash ); HashTable.IsValid(j); j = HashTable.Next(j) )
-		{
-			TSimpVert<T>* v1 = &sVerts[i];
-			TSimpVert<T>* v2 = &sVerts[j];
-
-			if( v1 == v2 )
-				continue;
-
-			// link
-			if( v1->GetPos() == v2->GetPos() )
-			{
-				checkSlow( v2->next == v2 );
-				checkSlow( v2->prev == v2 );
-
-				v2->next = v1->next;
-				v2->prev = v1;
-				v2->next->prev = v2;
-				v2->prev->next = v2;
-			}
-		}
-	}
-}
-
-template< typename T, uint32 NumAttributes >
-void TMeshSimplifier<T, NumAttributes>::GroupEdges()
-{
-	FHashTable HashTable( 1 << FMath::Min( 16u, FMath::FloorLog2( edges.Num() / 2 ) ), edges.Num() );
-	for( int i = 0; i < edges.Num(); i++ )
-	{
-		uint32 Hash0 = HashPoint( edges[i].v0->GetPos() );
-		uint32 Hash1 = HashPoint( edges[i].v1->GetPos() );
-		uint32 Hash = Murmur32( { FMath::Min( Hash0, Hash1 ), FMath::Max( Hash0, Hash1 ) } );
-		HashTable.Add( Hash, i );
-	}
-
-	for( int i = 0; i < edges.Num(); i++ )
-	{
-		// already grouped
-		if( edges[i].next != &edges[i] )
-		{
-			continue;
-		}
-
-		// find any matching edges
-		uint32 Hash0 = HashPoint( edges[i].v0->GetPos() );
-		uint32 Hash1 = HashPoint( edges[i].v1->GetPos() );
-		uint32 Hash = Murmur32( { FMath::Min( Hash0, Hash1 ), FMath::Max( Hash0, Hash1 ) } );
-		for( uint32 j = HashTable.First( Hash ); HashTable.IsValid(j); j = HashTable.Next(j) )
-		{
-			TSimpEdge<T>* e1 = &edges[i];
-			TSimpEdge<T>* e2 = &edges[j];
-
-			if( e1 == e2 )
-				continue;
-
-			bool m1 =	e1->v0->GetPos() == e2->v0->GetPos() &&
-						e1->v1->GetPos() == e2->v1->GetPos();
-
-			bool m2 =	e1->v0->GetPos() == e2->v1->GetPos() &&
-						e1->v1->GetPos() == e2->v0->GetPos();
-
-			// backwards
-			if( m2 )
-			{
-				Swap( e2->v0, e2->v1 );
-			}
-
-			// link
-			if( m1 || m2 )
-			{
-				check( e2->next == e2 );
-				check( e2->prev == e2 );
-
-				e2->next = e1->next;
-				e2->prev = e1;
-				e2->next->prev = e2;
-				e2->prev->next = e2;
-			}
-		}
-	}
-}
-
-template< typename T, uint32 NumAttributes >
-void TMeshSimplifier<T, NumAttributes>::InitCosts()
-{
-	for( auto& edge : edges )
-	{
-		AddEdgeToHeap( &edge );
-	}
-}
-
-template< typename T, uint32 NumAttributes >
-TQuadricAttr< NumAttributes >& TMeshSimplifier<T, NumAttributes>::GetQuadric( const TSimpTri<T>* Tri )
-{
-	uint32 TriIndex = GetTriIndex( Tri );
-	if( !TriQuadricsValid[ TriIndex ] )
-	{
-		TriQuadricsValid[ TriIndex ] = true;
-		TriQuadrics[ TriIndex ] = QuadricType(
-			Tri->verts[0]->GetPos(),		Tri->verts[1]->GetPos(),		Tri->verts[2]->GetPos(),
-			Tri->verts[0]->GetAttributes(),	Tri->verts[1]->GetAttributes(),	Tri->verts[2]->GetAttributes(),
-			attributeWeights );
-	}
-
-	return TriQuadrics[ TriIndex ];
-}
-
-template< typename T, uint32 NumAttributes >
-TQuadricAttr< NumAttributes >& TMeshSimplifier<T, NumAttributes>::GetQuadric( TSimpVert<T>* v )
-{
-	uint32 VertIndex = GetVertIndex( v );
-	if( !VertQuadricsValid[ VertIndex ] )
-	{
-		QuadricType vertQuadric;
-		vertQuadric.Zero();
-	
-		// sum tri quadrics
-		for( TriIterator i = v->adjTris.Begin(); i != v->adjTris.End(); ++i )
-		{
-			vertQuadric += GetQuadric( *i );
-		}
-
-		VertQuadricsValid[ VertIndex ] = true;
-		VertQuadrics[ VertIndex ] = vertQuadric;
-	}
-
-	return VertQuadrics[ VertIndex ];
-}
-
-template< typename T, uint32 NumAttributes >
-FQuadric& TMeshSimplifier<T, NumAttributes>::GetEdgeQuadric( TSimpVert<T>* v )
-{
-	uint32 VertIndex = GetVertIndex( v );
-	if( !EdgeQuadricsValid[ VertIndex ] )
-	{
-		FQuadric vertQuadric;
-		vertQuadric.Zero();
-
-		TArray< TSimpVert<T>*, TInlineAllocator<64> > adjVerts;
-		v->FindAdjacentVerts( adjVerts );
-
-		LockTriFlags( SIMP_MARK1 );
-		v->EnableAdjTriFlags( SIMP_MARK1 );
-	
-		for( TSimpVert<T>* vert : adjVerts )
-		{
-			TSimpTri<T>* face = NULL;
-			int faceCount = 0;
-			for( TriIterator j = vert->adjTris.Begin(); j != vert->adjTris.End(); ++j )
-			{
-				TSimpTri<T>* tri = *j;
-				if( tri->TestFlags( SIMP_MARK1 ) )
-				{
-					face = tri;
-					faceCount++;
-				}
-			}
-
-			if( faceCount == 1 )
-			{
-				// only one face on this edge
-				FEdgeQuadric edgeQuadric( v->GetPos(), vert->GetPos(), face->GetNormal(), EdgeWeight );
-				vertQuadric.Add( edgeQuadric, v->GetPos() );
-			}
-		}
-
-		v->DisableAdjTriFlags( SIMP_MARK1 );
-		UnlockTriFlags( SIMP_MARK1 );
-
-		EdgeQuadricsValid[ VertIndex ] = true;
-		EdgeQuadrics[ VertIndex ] = vertQuadric;
-	}
-
-	return EdgeQuadrics[ VertIndex ];
-}
-
-template< typename T, uint32 NumAttributes >
-FORCEINLINE uint32 TMeshSimplifier<T, NumAttributes>::GetVertIndex( const TSimpVert<T>* vert ) const
-{
-	ptrdiff_t Index = vert - &sVerts[0];
-	return (uint32)Index;
-}
-
-template< typename T, uint32 NumAttributes >
-FORCEINLINE uint32 TMeshSimplifier<T, NumAttributes>::GetTriIndex( const TSimpTri<T>* tri ) const
-{
-	ptrdiff_t Index = tri - &sTris[0];
-	return (uint32)Index;
-}
-
-template< typename T, uint32 NumAttributes >
-FORCEINLINE uint32 TMeshSimplifier<T, NumAttributes>::GetEdgeIndex( const TSimpEdge<T>* edge ) const
-{
-	ptrdiff_t Index = edge - &edges[0];
-	return (uint32)Index;
-}
-
-template< typename T, uint32 NumAttributes >
-FORCEINLINE uint32 TMeshSimplifier<T, NumAttributes>::HashPoint( const FVector3f& p ) const
-{
-	union { float f; uint32 i; } x;
-	union { float f; uint32 i; } y;
-	union { float f; uint32 i; } z;
-
-	x.f = p.X;
-	y.f = p.Y;
-	z.f = p.Z;
-
-	return Murmur32( { x.i, y.i, z.i } );
-}
-
-template< typename T, uint32 NumAttributes >
-FORCEINLINE uint32 TMeshSimplifier<T, NumAttributes>::HashEdge( const TSimpVert<T>* u, const TSimpVert<T>* v ) const
-{
-	uint32 ui = GetVertIndex( u );
-	uint32 vi = GetVertIndex( v );
-	// must be symmetrical
-	return Murmur32( { FMath::Min( ui, vi ), FMath::Max( ui, vi ) } );
-}
-
-template< typename T, uint32 NumAttributes >
-TSimpEdge<T>* TMeshSimplifier<T, NumAttributes>::FindEdge( const TSimpVert<T>* u, const TSimpVert<T>* v )
-{
-	uint32 hash = HashEdge( u, v );
-	for( uint32 i = edgeHash.First( hash ); edgeHash.IsValid(i); i = edgeHash.Next(i) )
-	{
-		if( ( edges[i].v0 == u && edges[i].v1 == v ) ||
-			( edges[i].v0 == v && edges[i].v1 == u ) )
-		{
-			return &edges[i];
-		}
-	}
-
-	return NULL;
-}
-
-template< typename T, uint32 NumAttributes >
-void TMeshSimplifier<T, NumAttributes>::RemoveEdge( TSimpEdge<T>* edge )
-{
-	if( edge->TestFlags( SIMP_REMOVED ) )
-	{
-		checkSlow( edge->next == edge );
-		checkSlow( edge->prev == edge );
-		return;
-	}
-
-	uint32 hash = HashEdge( edge->v0, edge->v1 );
-	for( uint32 i = edgeHash.First( hash ); edgeHash.IsValid(i); i = edgeHash.Next(i) )
-	{
-		if( &edges[i] == edge )
-		{
-			edgeHash.Remove( hash, i );
-			edgeHeap.Remove( i );
-			break;
-		}
-	}
-
-	// remove edge
-	edge->EnableFlags( SIMP_REMOVED );
-
-	// ungroup edge
-	edge->prev->next = edge->next;
-	edge->next->prev = edge->prev;
-	edge->next = edge;
-	edge->prev = edge;
-}
-
-template< typename T, uint32 NumAttributes >
-void TMeshSimplifier<T, NumAttributes>::RemoveVert( TSimpVert<T>* vert )
-{
-	vert->adjTris.Clear();
-	vert->EnableFlags( SIMP_REMOVED );
-
-	// ungroup
-	vert->prev->next = vert->next;
-	vert->next->prev = vert->prev;
-	vert->next = vert;
-	vert->prev = vert;
-	
-	numVerts--;
-}
-
-template< typename T, uint32 NumAttributes >
-void TMeshSimplifier<T, NumAttributes>::ReplaceEdgeVert( const TSimpVert<T>* oldV, const TSimpVert<T>* otherV, TSimpVert<T>* newV )
-{
-	uint32 hash = HashEdge( oldV, otherV );
-	uint32 index;
-	for( index = edgeHash.First( hash ); edgeHash.IsValid( index ); index = edgeHash.Next( index ) )
-	{
-		if( ( edges[ index ].v0 == oldV && edges[ index ].v1 == otherV ) ||
-			( edges[ index ].v1 == oldV && edges[ index ].v0 == otherV ) )
-			break;
-	}
-
-	checkSlow( index != -1 );
-	TSimpEdge<T>* edge = &edges[ index ];
-	
-	edgeHash.Remove( hash, index );
-
-	TSimpEdge<T>* ExistingEdge = FindEdge( newV, otherV );
-	if( ExistingEdge )
-	{
-		// Not entirely sure why this happens. I believe these are invalid edges from bridge tris.
-		RemoveEdge( ExistingEdge );
-	}
-
-	if( newV )
-	{
-		edgeHash.Add( HashEdge( newV, otherV ), index );
-
-		if( edge->v0 == oldV )
-			edge->v0 = newV;
-		else
-			edge->v1 = newV;
-	}
-	else
-	{
-		// remove edge
-		edge->EnableFlags( SIMP_REMOVED );
-
-		// ungroup old edge
-		edge->prev->next = edge->next;
-		edge->next->prev = edge->prev;
-		edge->next = edge;
-		edge->prev = edge;
-
-		edgeHeap.Remove( index );
-	}
-}
-
-template< typename T, uint32 NumAttributes >
-void TMeshSimplifier<T, NumAttributes>::CollapseEdgeVert( const TSimpVert<T>* oldV, const TSimpVert<T>* otherV, TSimpVert<T>* newV )
-{
-	uint32 hash = HashEdge( oldV, otherV );
-	uint32 index;
-	for( index = edgeHash.First( hash ); edgeHash.IsValid( index ); index = edgeHash.Next( index ) )
-	{
-		if( ( edges[ index ].v0 == oldV && edges[ index ].v1 == otherV ) ||
-			( edges[ index ].v1 == oldV && edges[ index ].v0 == otherV ) )
-			break;
-	}
-
-	checkSlow( index != -1 );
-	TSimpEdge<T>* edge = &edges[ index ];
-	
-	edgeHash.Remove( hash, index );
-	edgeHeap.Remove( index );
-
-	// remove edge
-	edge->EnableFlags( SIMP_REMOVED );
-
-	// ungroup old edge
-	edge->prev->next = edge->next;
-	edge->next->prev = edge->prev;
-	edge->next = edge;
-	edge->prev = edge;
-}
-
-template< typename T, uint32 NumAttributes >
-void TMeshSimplifier<T, NumAttributes>::GatherUpdates( TSimpVert<T>* v )
-{
-	TArray< TSimpVert<T>*, TInlineAllocator<64> > AdjVerts;
-	
-	for( TriIterator i = v->adjTris.Begin(); i != v->adjTris.End(); ++i )
-	{
-		TSimpTri<T>* Tri = *i;
-
-		// Update all tris touching collapse edge.
-		updateTris.AddUnique( Tri );
-
-		for( uint32 CornerIndex = 0; CornerIndex < 3; CornerIndex++ )
-		{
-			TSimpVert<T>* Vert = Tri->verts[ CornerIndex ];
-
-			if( AdjVerts.Find( Vert ) == INDEX_NONE )
-			{
-				AdjVerts.Add( Vert );
-
-				// Update verts from tris adjacent to collapsed edge
-				updateVerts.AddUnique( Vert );
-			}
-		}
-	}
-
-	LockVertFlags( SIMP_MARK1 | SIMP_MARK2 );
-
-	// update the costs of all edges connected to any face adjacent to v
-	for( TSimpVert<T>* AdjVert : AdjVerts )
-	{
-		// Mark as needing to be tested against this AdjVert
-		AdjVert->EnableAdjVertFlags( SIMP_MARK1 );
-
-		// Mark as processed as an AdjVert. No need to try edges connecting to it after we've evaluated all its possible connections already.
-		AdjVert->EnableFlags( SIMP_MARK2 );
-
-		for( TriIterator i = AdjVert->adjTris.Begin(); i != AdjVert->adjTris.End(); ++i )
-		{
-			TSimpTri<T>* Tri = *i;
-
-			for( uint32 CornerIndex = 0; CornerIndex < 3; CornerIndex++ )
-			{
-				TSimpVert<T>* Vert = Tri->verts[ CornerIndex ];
-
-				// Any vert that hasn't been tested against AdjVert yet nor has been an AdjVert itself already.
-				if( Vert->TestFlags( SIMP_MARK1 ) && !Vert->TestFlags( SIMP_MARK2 ) )
-				{
-					TSimpEdge<T>* edge = FindEdge( AdjVert, Vert );
-					updateEdges.AddUnique( edge );
-				}
-				Vert->DisableFlags( SIMP_MARK1 );
-			}
-		}
-	}
-
-	for( TSimpVert<T>* AdjVert : AdjVerts )
-	{
-		AdjVert->DisableFlags( SIMP_MARK2 );
-	}
-
-	UnlockVertFlags( SIMP_MARK1 | SIMP_MARK2 );
-}
-
-template< typename T, uint32 NumAttributes >
-float TMeshSimplifier<T, NumAttributes>::Merge( TSimpVert<T>* v0, TSimpVert<T>* v1, bool bMoveVerts )
-{
-	const FVector3f Position0 = v0->GetPos();
-	const FVector3f Position1 = v1->GetPos();
-
-	// Find unique adjacent triangles
-	TArray< TSimpTri<T>*, TInlineAllocator<16> > AdjTris;
-	
-	TSimpVert<T>* v;
-	v = v0;
-	do {
-		checkSlow( !v->TestFlags( SIMP_REMOVED ) );
-		checkSlow( v->adjTris.Num() != 0 );
-
-		for( TriIterator i = v->adjTris.Begin(); i != v->adjTris.End(); ++i )
-		{
-			TSimpTri<T>* Tri = *i;
-			checkSlow( !Tri->TestFlags( SIMP_REMOVED ) );
-
-			if( !Tri->TestFlags( SIMP_MARK1 ) )
-			{
-				Tri->EnableFlags( SIMP_MARK1 );
-				AdjTris.Add( Tri );
-			}
-		}
-		v = v->next;
-	} while( v != v0 );
-
-	v = v1;
-	do {
-		checkSlow( !v->TestFlags( SIMP_REMOVED ) );
-		checkSlow( v->adjTris.Num() != 0 );
-
-		for( TriIterator i = v->adjTris.Begin(); i != v->adjTris.End(); ++i )
-		{
-			TSimpTri<T>* Tri = *i;
-			checkSlow( !Tri->TestFlags( SIMP_REMOVED ) );
-
-			if( !Tri->TestFlags( SIMP_MARK1 ) )
-			{
-				Tri->EnableFlags( SIMP_MARK1 );
-				AdjTris.Add( Tri );
-			}
-		}
-		v = v->next;
-	} while( v != v1 );
-
-	FVector3f	BoundsMin = {  MAX_flt,  MAX_flt,  MAX_flt };
-	FVector3f	BoundsMax = { -MAX_flt, -MAX_flt, -MAX_flt };
-
-	AdjCorners.Clear();
-	WedgeSet.Init( AdjTris.Num() );
-
-	for( uint32 AdjTriIndex = 0, Num = AdjTris.Num(); AdjTriIndex < Num; AdjTriIndex++ )
-	{
-		for( uint32 CornerIndex = 0; CornerIndex < 3; CornerIndex++ )
-		{
-			TSimpVert<T>* RESTRICT Vert = AdjTris[ AdjTriIndex ]->verts[ CornerIndex ];
-			
-			uint32 Hash = HashPoint( Vert->GetPos() );
-			uint32 AdjCorner;
-			for( AdjCorner = AdjCorners.First( Hash ); AdjCorners.IsValid( AdjCorner ); AdjCorner = AdjCorners.Next( AdjCorner ) )
-			{
-				uint32 OtherAdjTriIndex = AdjCorner >> 2;
-				uint32 OtherCornerIndex = AdjCorner & 3;
-
-				TSimpVert<T>* OtherVert = AdjTris[ OtherAdjTriIndex ]->verts[ OtherCornerIndex ];
-				if( Vert == OtherVert )
-				{
-					// Only union wedges with the merged positions. Shared adjverts elsewhere are irrelevant for wedges.
-					if( Vert->GetPos() == Position0 ||
-						Vert->GetPos() == Position1 )
-					{
-						WedgeSet.UnionSequential( AdjTriIndex, OtherAdjTriIndex );
-					}
-
-					break;
-				}
-			}
-			if( !AdjCorners.IsValid( AdjCorner ) )
-			{
-				// if didn't find this vert already present, add it.
-				AdjCorners.Add( Hash, (AdjTriIndex << 2) | CornerIndex );
-
-				BoundsMin = FVector3f::Min( BoundsMin, Vert->GetPos() );
-				BoundsMax = FVector3f::Max( BoundsMax, Vert->GetPos() );
-			}
-		}
-	}
-
-	float BoundsExtentSqr = ( BoundsMax - BoundsMin ).SizeSquared();
-	
-	TArray< uint32, TInlineAllocator<8> >		WedgeIDs;
-	TArray< QuadricType, TInlineAllocator<8> >	WedgeQuadrics;
-
-	for( uint32 AdjTriIndex = 0, Num = AdjTris.Num(); AdjTriIndex < Num; AdjTriIndex++ )
-	{
-		TSimpTri<T>* RESTRICT Tri = AdjTris[ AdjTriIndex ];
-		Tri->DisableFlags( SIMP_MARK1 );
-
-		int32 QuadricIndex = WedgeIDs.AddUnique( WedgeSet.Find( AdjTriIndex ) );
-		if( QuadricIndex < WedgeQuadrics.Num() )
-		{
-			WedgeQuadrics[ QuadricIndex ] += GetQuadric( Tri );
-		}
-		else
-		{
-			WedgeQuadrics.Add( GetQuadric( Tri ) );
-		}
-	}
-
-	FQuadricAttrOptimizer QuadricOptimizer;
-	for( auto& WedgeQuadric : WedgeQuadrics )
-	{
-		QuadricOptimizer.AddQuadric( WedgeQuadric, NumAttributes );
-	}
-
-	// TODO edge quadrics differently
-	FQuadric EdgeQuadric;
-	EdgeQuadric.Zero();
-
-	v = v0;
-	do {
-		EdgeQuadric += GetEdgeQuadric(v);
-		v = v->next;
-	} while( v != v0 );
-
-	v = v1;
-	do {
-		EdgeQuadric += GetEdgeQuadric(v);
-		v = v->next;
-	} while( v != v1 );
-
-	QuadricOptimizer.AddQuadric( EdgeQuadric );
-
-
-	FVector3f NewPos;
-	{
-		bool bLocked0 = v0->TestFlags( SIMP_LOCKED );
-		bool bLocked1 = v1->TestFlags( SIMP_LOCKED );
-
-		// find position
-		if( bLocked0 && !bLocked1 )
-		{
-			// v0 position
-			NewPos = v0->GetPos();
-		}
-		else if( bLocked1 && !bLocked0 )
-		{
-			// v1 position
-			NewPos = v1->GetPos();
-		}
-		else
-		{
-			bool bIsValid = QuadricOptimizer.OptimizeVolume( NewPos );
-
-			// Limit position to be near the neighborhood bounds
-			if( !bIsValid || ComputeSquaredDistanceFromBoxToPoint( BoundsMin, BoundsMax, NewPos ) > BoundsExtentSqr * 4.0f )
-			{
-				bIsValid = QuadricOptimizer.Optimize( NewPos );
-			}
-			
-			if( !bIsValid || ComputeSquaredDistanceFromBoxToPoint( BoundsMin, BoundsMax, NewPos ) > BoundsExtentSqr * 4.0f )
-			{
-				// Try a point on the edge.
-				bIsValid = QuadricOptimizer.OptimizeLinear( Position0, Position1, NewPos );
-			}
-
-			if( !bIsValid )
-			{
-				// Couldn't find optimal so choose middle
-				NewPos = ( Position0 + Position1 ) * 0.5f;
-			}
-		}
-	}
-
-
-	TArray< T, TInlineAllocator<8> > NewVerts;
-	NewVerts.Reserve( WedgeQuadrics.Num() );
-	
-	float Error = 0.0f;
-	float EdgeError = EdgeQuadric.Evaluate( NewPos );
-
-	for( auto& WedgeQuadric : WedgeQuadrics )
-	{
-		T& NewVert = NewVerts.AddDefaulted_GetRef();
-		NewVert.GetPos() = NewPos;
-
-		if( WedgeQuadric.a > 1e-8 )
-		{
-			// calculate vert attributes from the new position
-			float WedgeError = WedgeQuadric.CalcAttributesAndEvaluate( NewPos, NewVert.GetAttributes(), attributeWeights );
-			
-			// Correct after eval. Normal length is unimportant for error but can bias the calculation.
-			NewVert.Correct();
-
-			Error += WedgeError;
-		}
-		else
-		{
-			for( uint32 i = 0; i < NumAttributes; i++ )
-			{
-				NewVert.GetAttributes()[i] = 0.0f;
-			}
-		}
-	}
-
-	if( bMoveVerts )
-	{
-		for( uint32 AdjTriIndex = 0, Num = AdjTris.Num(); AdjTriIndex < Num; AdjTriIndex++ )
-		{
-			int32 QuadricIndex = WedgeIDs.Find( WedgeSet[ AdjTriIndex ] );
-
-			for( uint32 CornerIndex = 0; CornerIndex < 3; CornerIndex++ )
-			{
-				TSimpVert<T>* Vert = AdjTris[ AdjTriIndex ]->verts[ CornerIndex ];
-
-				if( Vert->GetPos() == Position0 ||
-					Vert->GetPos() == Position1 )
-				{
-					// Only use attributes if we calculated them.
-					if( WedgeQuadrics[ QuadricIndex ].a > 1e-8 )
-						Vert->vert = NewVerts[ QuadricIndex ];
-					else
-						Vert->GetPos() = NewPos;
-				}
-			}
-		}
-		
-		// Limit error to be no greater than the size of the triangles it could affect.
-		Error = FMath::Min( Error, BoundsExtentSqr );
-	}
-	else
-	{
-		Error += EdgeError;
-		//Error = FMath::Min( Error, BoundsExtentSqr );
-		Error += ComputeMergePenalty( v0, v1, NewPos );
-	}
-
-	return Error;
-}
-
-#if !SIMP_NEW_MERGE
-template< typename T, uint32 NumAttributes >
-float TMeshSimplifier<T, NumAttributes>::ComputeNewVerts( TSimpEdge<T>* edge, TArray< T, TInlineAllocator<16> >& newVerts )
-{
-	TSimpEdge<T>* e;
-	TSimpVert<T>* v;
-
-	TArray< QuadricType, TInlineAllocator<16> > quadrics;
-
-	FQuadricAttrOptimizer optimizer;
-
-	LockVertFlags( SIMP_MARK1 );
-	
-	edge->v0->EnableFlagsGroup( SIMP_MARK1 );
-	edge->v1->EnableFlagsGroup( SIMP_MARK1 );
-
-	// add edges
-	e = edge;
-	do {
-		checkSlow( e == FindEdge( e->v0, e->v1 ) );
-		checkSlow( e->v0->adjTris.Num() > 0 );
-		checkSlow( e->v1->adjTris.Num() > 0 );
-
-		newVerts.Add( e->v0->vert );
-
-		QuadricType quadric;
-		quadric  = GetQuadric( e->v0 );
-		quadric += GetQuadric( e->v1 );
-		quadrics.Add( quadric );
-		optimizer.AddQuadric( quadric );
-
-		e->v0->DisableFlags( SIMP_MARK1 );
-		e->v1->DisableFlags( SIMP_MARK1 );
-
-		e = e->next;
-	} while( e != edge );
-
-	// add remainder verts
-	v = edge->v0;
-	do {
-		if( v->TestFlags( SIMP_MARK1 ) )
-		{
-			checkSlow( !v->TestFlags( SIMP_REMOVED ) );
-			checkSlow( v->adjTris.Num() != 0 );
-
-			newVerts.Add( v->vert );
-
-			QuadricType quadric;
-			quadric = GetQuadric( v );
-			quadrics.Add( quadric );
-			optimizer.AddQuadric( quadric );
-
-			v->DisableFlags( SIMP_MARK1 );
-		}
-		v = v->next;
-	} while( v != edge->v0 );
-
-	v = edge->v1;
-	do {
-		if( v->TestFlags( SIMP_MARK1 ) )
-		{
-			checkSlow( !v->TestFlags( SIMP_REMOVED ) );
-			checkSlow( v->adjTris.Num() != 0 );
-
-			newVerts.Add( v->vert );
-
-			QuadricType quadric;
-			quadric = GetQuadric( v );
-			quadrics.Add( quadric );
-			optimizer.AddQuadric( quadric );
-
-			v->DisableFlags( SIMP_MARK1 );
-		}
-		v = v->next;
-	} while( v != edge->v1 );
-
-	UnlockVertFlags( SIMP_MARK1 );
-
-	FQuadric edgeQuadric;
-	edgeQuadric.Zero();
-
-	v = edge->v0;
-	do {
-		edgeQuadric += GetEdgeQuadric( v );
-		v = v->next;
-	} while( v != edge->v0 );
-
-	v = edge->v1;
-	do {
-		edgeQuadric += GetEdgeQuadric( v );
-		v = v->next;
-	} while( v != edge->v1 );
-
-	optimizer.AddQuadric( edgeQuadric );
-	
-	FVector3f newPos;
-	{
-		bool bLocked0 = edge->v0->TestFlags( SIMP_LOCKED );
-		bool bLocked1 = edge->v1->TestFlags( SIMP_LOCKED );
-		//checkSlow( !bLocked0 || !bLocked1 );
-
-		// find position
-		if( bLocked0 && !bLocked1 )
-		{
-			// v0 position
-			newPos = edge->v0->GetPos();
-		}
-		else if( bLocked1 && !bLocked0 )
-		{
-			// v1 position
-			newPos = edge->v1->GetPos();
-		}
-		else
-		{
-			// optimal position
-			FVector3f pos;
-			bool valid = optimizer.Optimize( pos, newPos );
-			if( !valid )
-			{
-				// Couldn't find optimal so choose middle
-				newPos = ( edge->v0->GetPos() + edge->v1->GetPos() ) * 0.5f;
-			}
-		}
-	}
-	
-	float cost = 0.0f;
-
-	for( int i = 0; i < quadrics.Num(); i++ )
-	{
-		newVerts[i].GetPos() = newPos;
-
-		if( quadrics[i].a > 1e-8 )
-		{
-			// calculate vert attributes from the new position
-			quadrics[i].CalcAttributesAndEvaluate( newVerts[i].GetPos(), newVerts[i].GetAttributes(), attributeWeights );
-			newVerts[i].Correct();
-		}
-
-		// sum cost of new verts
-		cost += quadrics[i].Evaluate( newVerts[i].GetPos(), newVerts[i].GetAttributes(), attributeWeights );
-	}
-
-	cost += edgeQuadric.Evaluate( newPos );
-
-	return cost;
-}
-
-template< typename T, uint32 NumAttributes >
-float TMeshSimplifier<T, NumAttributes>::ComputeEdgeCollapseCost( TSimpEdge<T>* edge, bool bMoveVerts )
-{
-	TArray< T, TInlineAllocator<16> > newVerts;
-	float Error = ComputeNewVerts( edge, newVerts );
-
-	if( bMoveVerts )
-	{
-		TSimpVert<T>* v;
-
-		uint32 i = 0;
-
-		LockVertFlags( SIMP_MARK1 );
-
-		edge->v0->EnableFlagsGroup( SIMP_MARK1 );
-		edge->v1->EnableFlagsGroup( SIMP_MARK1 );
-
-		// edges
-		TSimpEdge<T>* e = edge;
-		do {
-			checkSlow( e == FindEdge( e->v0, e->v1 ) );
-			checkSlow( e->v0->adjTris.Num() > 0 );
-			checkSlow( e->v1->adjTris.Num() > 0 );
-
-			e->v1->vert = newVerts[ i++ ];
-			e->v0->DisableFlags( SIMP_MARK1 );
-			e->v1->DisableFlags( SIMP_MARK1 );
-
-			e = e->next;
-		} while( e != edge );
-
-		// remainder verts
-		v = edge->v0;
-		do {
-			if( v->TestFlags( SIMP_MARK1 ) )
-			{
-				checkSlow( !v->TestFlags( SIMP_REMOVED ) );
-				checkSlow( v->adjTris.Num() != 0 );
-
-				v->vert = newVerts[ i++ ];
-				v->DisableFlags( SIMP_MARK1 );
-			}
-			v = v->next;
-		} while( v != edge->v0 );
-
-		v = edge->v1;
-		do {
-			if( v->TestFlags( SIMP_MARK1 ) )
-			{
-				checkSlow( !v->TestFlags( SIMP_REMOVED ) );
-				checkSlow( v->adjTris.Num() != 0 );
-
-				v->vert = newVerts[ i++ ];
-				v->DisableFlags( SIMP_MARK1 );
-			}
-			v = v->next;
-		} while( v != edge->v1 );
-
-		UnlockVertFlags( SIMP_MARK1 );
-	}
-	else
-	{
-		Error += ComputeMergePenalty( edge->v0, edge->v1, newVerts[0].GetPos() );
-	}
-
-	return Error;
-}
-#endif
-
-template< typename T, uint32 NumAttributes >
-float TMeshSimplifier<T, NumAttributes>::ComputeMergePenalty( TSimpVert<T>* v0, TSimpVert<T>* v1, const FVector3f& NewPos )
-{
-	// add penalties
-	// the below penalty code works with groups so no need to worry about remainder verts
-	TSimpVert<T>* v;
-
-	float penalty = 0.0f;
-
-	{
-		const int degreeLimit = 24;
-		const float degreePenalty = 100.0f;
-
-		int degree = 0;
-
-		v = v0;
-		do {
-			degree += v->adjTris.Num();
-			v = v->next;
-		} while( v != v0 );
-
-		// v
-		v = v1;
-		do {
-			degree += v->adjTris.Num();
-			v = v->next;
-		} while( v != v1 );
-
-		if( degree > degreeLimit )
-			penalty += degreePenalty * float( degree - degreeLimit );
-	}
-
-	{
-		const float LockPenalty = 1e8f;
-		if( v0->TestFlags( SIMP_LOCKED ) && v1->TestFlags( SIMP_LOCKED ) )
-			penalty += LockPenalty;
-	}
-
-	{
-		// Penalty to prevent edge folding 
-		const float invalidPenalty = 1000000.0f;
-
-		LockTriFlags( SIMP_MARK1 );
-	
-		v1->EnableAdjTriFlagsGroup( SIMP_MARK1 );
-
-		// u
-		v = v0;
-		do {
-			for( TriIterator i = v->adjTris.Begin(); i != v->adjTris.End(); ++i )
-			{
-				TSimpTri<T>* tri = *i;
-				if( !tri->TestFlags( SIMP_MARK1 ) )
-					penalty += tri->ReplaceVertexIsValid( v, NewPos ) ? 0.0f : invalidPenalty;
-				tri->DisableFlags( SIMP_MARK1 );
-			}
-			v = v->next;
-		} while( v != v0 );
-
-		// v
-		v = v1;
-		do {
-			for( TriIterator i = v->adjTris.Begin(); i != v->adjTris.End(); ++i )
-			{
-				TSimpTri<T>* tri = *i;
-				if( tri->TestFlags( SIMP_MARK1 ) )
-					penalty += tri->ReplaceVertexIsValid( v, NewPos ) ? 0.0f : invalidPenalty;
-				tri->DisableFlags( SIMP_MARK1 );
-			}
-			v = v->next;
-		} while( v != v1 );
-
-		UnlockTriFlags( SIMP_MARK1 );
-	}
-
-	return penalty;
-}
-
-template< typename T, uint32 NumAttributes >
-void TMeshSimplifier<T, NumAttributes>::AddEdgeToHeap( TSimpEdge<T>* edge )
-{
-	if( edgeHeap.IsPresent( GetEdgeIndex( edge ) ) )
-		return;
-
-#if SIMP_NEW_MERGE
-	float cost = Merge( edge->v0, edge->v1, false );
-#else
-	float cost = ComputeEdgeCollapseCost( edge, false );
-#endif
-
-	TSimpEdge<T>* e = edge;
-	do {
-		uint32 EdgeIndex = GetEdgeIndex(e);
-		edgeHeap.Add( cost, EdgeIndex );
-		e = e->next;
-	} while( e != edge );
-}
-
-template< typename T, uint32 NumAttributes >
-void TMeshSimplifier<T, NumAttributes>::Collapse( TSimpEdge<T>* edge )
-{
-	TSimpVert<T>* u = edge->v0;
-	TSimpVert<T>* v = edge->v1;
-	
-	// Collapse the edge uv by moving vertex u onto v
-	checkSlow( u && v );
-	checkSlow( edge == FindEdge( u, v ) );
-	checkSlow( u->adjTris.Num() > 0 );
-	checkSlow( v->adjTris.Num() > 0 );
-
-	if( u->TestFlags( SIMP_LOCKED ) )
-		v->EnableFlags( SIMP_LOCKED );
-
-	LockVertFlags( SIMP_MARK1 );
-
-	// update edges from u to v
-	u->EnableAdjVertFlags( SIMP_MARK1 );
-	v->DisableAdjVertFlags( SIMP_MARK1 );
-
-	if( u->TestFlags( SIMP_MARK1 ) )
-	{
-		// Invalid edge, results from collapsing a bridge tri
-		// There are no actual triangles connecting these verts
-		u->DisableAdjVertFlags( SIMP_MARK1 );
-		UnlockVertFlags( SIMP_MARK1 );
-		return;
-	}
-
-	for( TriIterator i = u->adjTris.Begin(); i != u->adjTris.End(); ++i )
-	{
-		TSimpTri<T>* tri = *i;
-		for( int j = 0; j < 3; j++ )
-		{
-			TSimpVert<T>* vert = tri->verts[j];
-			if( vert->TestFlags( SIMP_MARK1 ) )
-			{
-				ReplaceEdgeVert( u, vert, v );
-				vert->DisableFlags( SIMP_MARK1 );
-			}
-		}
-	}
-
-	// remove dead edges
-	u->EnableAdjVertFlags( SIMP_MARK1 );
-	u->DisableFlags( SIMP_MARK1 );
-	v->DisableFlags( SIMP_MARK1 );
-
-	for( TriIterator i = v->adjTris.Begin(); i != v->adjTris.End(); ++i )
-	{
-		TSimpTri<T>* tri = *i;
-		for( int j = 0; j < 3; j++ )
-		{
-			TSimpVert<T>* vert = tri->verts[j];
-			if( vert->TestFlags( SIMP_MARK1 ) )
-			{
-				ReplaceEdgeVert( u, vert, NULL );
-				vert->DisableFlags( SIMP_MARK1 );
-			}
-		}
-	}
-
-	u->DisableAdjVertFlags( SIMP_MARK1 );
-
-	// fixup triangles
-	for( TriIterator i = u->adjTris.Begin(); i != u->adjTris.End(); ++i )
-	{
-		TSimpTri<T>* tri = *i;
-		checkSlow( !tri->TestFlags( SIMP_REMOVED ) );
-		checkSlow( tri->HasVertex(u) );
-		
-		if( tri->HasVertex(v) )
-		{
-			// delete triangles on edge uv
-			numTris--;
-			tri->EnableFlags( SIMP_REMOVED );
-			TriQuadricsValid[ GetTriIndex( tri ) ] = false;
-
-			// remove references to tri
-			for( int j = 0; j < 3; j++ )
-			{
-				TSimpVert<T>* vert = tri->verts[j];
-				checkSlow( !vert->TestFlags( SIMP_REMOVED ) );
-				if( vert != u )
-				{
-					vert->adjTris.Remove( tri );
-					updateVerts.AddUnique( vert );
-				}
-			}
-		}
-		else
-		{
-			// update triangles to have v instead of u
-			tri->ReplaceVertex( u, v );
-			v->adjTris.Add( tri );
-		}
-	}
-	
-	UnlockVertFlags( SIMP_MARK1 );
-	
-	u->adjTris.Clear();	// u has been removed
-	updateVerts.AddUnique( u );
-}
-
-template< typename T, uint32 NumAttributes >
-void TMeshSimplifier<T, NumAttributes>::UpdateTris()
-{
-	// remove degenerate triangles
-	// not sure why this happens
-	for( TSimpTri<T>* tri : updateTris )
-	{
-		if( tri->TestFlags( SIMP_REMOVED ) )
-			continue;
-
-		TriQuadricsValid[ GetTriIndex( tri ) ] = false;
-
-		const FVector3f& p0 = tri->verts[0]->GetPos();
-		const FVector3f& p1 = tri->verts[1]->GetPos();
-		const FVector3f& p2 = tri->verts[2]->GetPos();
-
-		if( p0 == p1 ||
-			p1 == p2 ||
-			p2 == p0 )
-		{
-			numTris--;
-			tri->EnableFlags( SIMP_REMOVED );
-
-			// remove references to tri
-			for( int j = 0; j < 3; j++ )
-			{
-				TSimpVert<T>* vert = tri->verts[j];
-				vert->adjTris.Remove( tri );
-
-				updateVerts.AddUnique( vert );
-				// orphaned verts are removed below
-			}
-		}
-	}
-	updateTris.Reset();
-}
-
-template< typename T, uint32 NumAttributes >
-void TMeshSimplifier<T, NumAttributes>::UpdateVerts()
-{
-	// remove orphaned verts
-	for( TSimpVert<T>* vert : updateVerts )
-	{
-		if( vert->TestFlags( SIMP_REMOVED ) )
-			continue;
-
-		VertQuadricsValid[ GetVertIndex( vert ) ] = false;
-		EdgeQuadricsValid[ GetVertIndex( vert ) ] = false;
-
-		if( vert->adjTris.Num() == 0 )
-		{
-			RemoveVert( vert );
-		}
-	}
-	updateVerts.Reset();
-}
-
-template< typename T, uint32 NumAttributes >
-void TMeshSimplifier<T, NumAttributes>::UpdateEdges()
-{
-	// add all grouped edges
-	for( uint32 i = 0, Num = updateEdges.Num(); i < Num; i++ )
-	{
-		TSimpEdge<T>* edge = updateEdges[i];
-
-		if( edge->TestFlags( SIMP_REMOVED ) )
-			continue;
-
-		TSimpEdge<T>* e = edge;
-		do {
-			updateEdges.AddUnique(e);
-			e = e->next;
-		} while( e != edge );
-	}
-		
-	// remove dead edges
-	for( TSimpEdge<T>* edge : updateEdges )
-	{
-		if( edge->TestFlags( SIMP_REMOVED ) )
-			continue;
-
-		if( edge->v0->TestFlags( SIMP_REMOVED ) ||
-			edge->v1->TestFlags( SIMP_REMOVED ) )
-		{
-			RemoveEdge( edge );
-			continue;
-		}
-	}
-
-	// Fix edge groups
-	{
-		FHashTable HashTable( 128, updateEdges.Num() );
-
-		// ungroup edges
-		for( uint32 i = 0, Num = updateEdges.Num(); i < Num; i++ )
-		{
-			TSimpEdge<T>* edge = updateEdges[i];
-
-			if( edge->TestFlags( SIMP_REMOVED ) )
-				continue;
-			
-			edge->next = edge;
-			edge->prev = edge;
-
-			edgeHeap.Remove( GetEdgeIndex( edge ) );
-
-			uint32 Hash0 = HashPoint( edge->v0->GetPos() );
-			uint32 Hash1 = HashPoint( edge->v1->GetPos() );
-			uint32 Hash = Murmur32( { FMath::Min( Hash0, Hash1 ), FMath::Max( Hash0, Hash1 ) } );
-				
-			HashTable.Add( Hash, i );
-		}
-
-		// regroup edges
-		for( TSimpEdge<T>* edge : updateEdges )
-		{
-			if( edge->TestFlags( SIMP_REMOVED ) )
-				continue;
-
-			// already grouped
-			if( edge->next != edge )
-				continue;
-
-			// find any matching edges
-			uint32 Hash0 = HashPoint( edge->v0->GetPos() );
-			uint32 Hash1 = HashPoint( edge->v1->GetPos() );
-			uint32 Hash = Murmur32( { FMath::Min( Hash0, Hash1 ), FMath::Max( Hash0, Hash1 ) } );
-			for( uint32 j = HashTable.First( Hash ); HashTable.IsValid(j); j = HashTable.Next( j ) )
-			{
-				TSimpEdge<T>* e1 = edge;
-				TSimpEdge<T>* e2 = updateEdges[j];
-
-				if( e1 == e2 )
-					continue;
-
-				bool m1 =	e1->v0->GetPos() == e2->v0->GetPos() &&
-							e1->v1->GetPos() == e2->v1->GetPos();
-
-				bool m2 =	e1->v0->GetPos() == e2->v1->GetPos() &&
-							e1->v1->GetPos() == e2->v0->GetPos();
-
-				// backwards
-				if( m2 )
-					Swap( e2->v0, e2->v1 );
-
-				// link
-				if( m1 || m2 )
-				{
-					checkSlow( e2->next == e2 );
-					checkSlow( e2->prev == e2 );
-
-					e2->next = e1->next;
-					e2->prev = e1;
-					e2->next->prev = e2;
-					e2->prev->next = e2;
-				}
-			}
-		}
-	}
-
-	// update edges
-	for( TSimpEdge<T>* edge : updateEdges )
-	{
-		if( edge->TestFlags( SIMP_REMOVED ) )
-			continue;
-		
-		AddEdgeToHeap( edge );
-	}
-
-	updateEdges.Reset();
-}
-
-template< typename T, uint32 NumAttributes >
-float TMeshSimplifier<T, NumAttributes>::SimplifyMesh( float maxErrorLimit, int minTris, int minVerts )
-{
-	TRACE_CPUPROFILER_EVENT_SCOPE(TMeshSimplifier::SimplifyMesh);
-	TSimpVert<T>* v;
-
-	float maxError = 0.0f;
-
-	while( edgeHeap.Num() > 0 )
-	{
-		if( numTris <= minTris || numVerts <= minVerts )
-			break;
-
-		// get the next vertex to collapse
-		uint32 TopIndex = edgeHeap.Top();
-
-		float error = edgeHeap.GetKey( TopIndex );
-		if( error > maxErrorLimit )
-		{
-			break;
-		}
-		maxError = FMath::Max( maxError, error );
-		
-		edgeHeap.Pop();
-
-		TSimpEdge<T>* top = &edges[ TopIndex ];
-
-		int numEdges = 0;
-		TSimpEdge<T>* edgeList[32];
-
-		TSimpEdge<T>* edge = top;
-		do {
-			edgeList[ numEdges++ ] = edge;
-			edge = edge->next;
-		} while( edge != top );
-
-		check(top);
-
-		// skip locked edges
-		bool locked = false;
-		for( int i = 0; i < numEdges; i++ )
-		{
-			edge = edgeList[i];
-			if( edge->v0->TestFlags( SIMP_LOCKED ) && edge->v1->TestFlags( SIMP_LOCKED ) )
-			{
-				locked = true;
-			}
-		}
-		if( locked )
-			continue;
-
-		v = top->v0;
-		do {
-			GatherUpdates( v );
-			v = v->next;
-		} while( v != top->v0 );
-
-		v = top->v1;
-		do {
-			GatherUpdates( v );
-			v = v->next;
-		} while( v != top->v1 );
-
-#if 1
-		// remove edges with already removed verts
-		// not sure why this happens
-		for( int i = 0; i < numEdges; i++ )
-		{
-			if( edgeList[i]->v0->adjTris.Num() == 0 ||
-				edgeList[i]->v1->adjTris.Num() == 0 )
-			{
-				RemoveEdge( edgeList[i] );
-				edgeList[i] = NULL;
-			}
-			else
-			{
-				checkSlow( !edgeList[i]->TestFlags( SIMP_REMOVED ) );
-			}
-		}
-		if( top->v0->adjTris.Num() == 0 || top->v1->adjTris.Num() == 0 )
-			continue;
-#endif
-
-		// move verts to new verts
-		{	
-			edge = top;
-			
-#if SIMP_NEW_MERGE
-			float Error = Merge( edge->v0, edge->v1, true );
-#else
-			float Error = ComputeEdgeCollapseCost( edge, true );
-#endif
-			//maxError = FMath::Max( maxError, Error );
-		}
-
-		// collapse all edges
-		for( int i = 0; i < numEdges; i++ )
-		{
-			edge = edgeList[i];
-			if( !edge )
-				continue;
-			if( edge->TestFlags( SIMP_REMOVED ) )	// wtf?
-				continue;
-			if( edge->v0->adjTris.Num() == 0 )
-				continue;
-			if( edge->v1->adjTris.Num() == 0 )
-				continue;
-
-			Collapse( edge );
-			RemoveEdge( edge );
-		}
-
-		// add v0 remainder verts to v1
-		{
-			// combine v0 and v1 groups
-			top->v0->next->prev = top->v1->prev;
-			top->v1->prev->next = top->v0->next;
-			top->v0->next = top->v1;
-			top->v1->prev = top->v0;
-
-			// ungroup removed verts
-			TArray< TSimpVert<T>*, TInlineAllocator<16> > vertList;
-			
-			v = top->v1;
-			do {
-				vertList.Add(v);
-				v = v->next;
-			} while( v != top->v1 );
-			
-			for( TSimpVert<T>* u : vertList )
-			{
-				if( u->TestFlags( SIMP_REMOVED ) )
-				{
-					// ungroup
-					u->prev->next = u->next;
-					u->next->prev = u->prev;
-					u->next = u;
-					u->prev = u;
-				}
-			}
-		}
-
-		{
-			// spread locked flag to vert group
-			uint32 flags = 0;
-
-			v = top->v1;
-			do {
-				flags |= v->flags & SIMP_LOCKED;
-				v = v->next;
-			} while( v != top->v1 );
-
-			v = top->v1;
-			do {
-				v->flags |= flags;
-				v = v->next;
-			} while( v != top->v1 );
-		}
-
-		UpdateTris();
-		UpdateVerts();
-		UpdateEdges();
-	}
-
-	// remove degenerate triangles
-	// not sure why this happens
-	for( int i = 0; i < numSTris; i++ )
-	{
-		TSimpTri<T>* tri = &sTris[i];
-
-		if( tri->TestFlags( SIMP_REMOVED ) )
-			continue;
-
-		const FVector3f& p0 = tri->verts[0]->GetPos();
-		const FVector3f& p1 = tri->verts[1]->GetPos();
-		const FVector3f& p2 = tri->verts[2]->GetPos();
-
-		if( p0 == p1 ||
-			p1 == p2 ||
-			p2 == p0 )
-		{
-			numTris--;
-			tri->EnableFlags( SIMP_REMOVED );
-
-			// remove references to tri
-			for( int j = 0; j < 3; j++ )
-			{
-				TSimpVert<T>* vert = tri->verts[j];
-				vert->adjTris.Remove( tri );
-				// orphaned verts are removed below
-			}
-		}
-	}
-
-	// remove orphaned verts
-	for( int i = 0; i < numSVerts; i++ )
-	{
-		TSimpVert<T>* vert = &sVerts[i];
-
-		if( vert->TestFlags( SIMP_REMOVED ) )
-			continue;
-
-		if( vert->adjTris.Num() == 0 )
-		{
-			RemoveVert( vert );
-		}
-	}
-
-	return maxError;
-}
-
-template< typename T, uint32 NumAttributes >
-void TMeshSimplifier<T, NumAttributes>::OutputMesh( T* verts, uint32* indexes, uint32* OutNumVertices, uint32* OutNumIndices ) const
-{
-	FHashTable HashTable( 1 << FMath::Min( 16u, FMath::FloorLog2( GetNumVerts() ) ), GetNumVerts() );
-
-#if 1
-	int count = 0;
-	for( int i = 0; i < numSVerts; i++ )
-		count += sVerts[i].TestFlags( SIMP_REMOVED ) ? 0 : 1;
-
-	check( numVerts == count );
-#endif
-
-	int numV = 0;
-	int numI = 0;
-
-	for( int i = 0; i < numSTris; i++ )
-	{
-		if( sTris[i].TestFlags( SIMP_REMOVED ) )
-			continue;
-
-		// TODO this is sloppy. There should be no duped verts. Alias by index.
-		for( int j = 0; j < 3; j++ )
-		{
-			TSimpVert<T>* vert = sTris[i].verts[j];
-			checkSlow( !vert->TestFlags( SIMP_REMOVED ) );
-			checkSlow( vert->adjTris.Num() != 0 );
-
-			const FVector3f& p = vert->GetPos();
-			uint32 hash = HashPoint( p );
-			uint32 f;
-			for( f = HashTable.First( hash ); HashTable.IsValid(f); f = HashTable.Next( f ) )
-			{
-				if( vert->vert == verts[f] )
-					break;
-			}
-			if( !HashTable.IsValid(f) )
-			{
-				HashTable.Add( hash, numV );
-				verts[ numV ] = vert->vert;
-				indexes[ numI++ ] = numV;
-				numV++;
-			}
-			else
-			{
-				indexes[ numI++ ] = f;
-			}
-		}
-	}
-
-	check( numV <= numVerts );
-	check( numI <= numTris * 3 );
-	
-	if (OutNumVertices)
-	{
-		*OutNumVertices = numV;
-	}
-
-	if (OutNumIndices)
-	{
-		*OutNumIndices = numI;
-	}
-}
-
-template< typename T, uint32 NumAttributes >
-template< typename FaceType >
-void TMeshSimplifier<T, NumAttributes>::CompactFaceData( TArray< FaceType>& FaceData )
-{
-	TArray< FaceType> NewFaceData;
-
-	for( int i = 0; i < numSTris; i++ )
-	{
-		if( sTris[i].TestFlags( SIMP_REMOVED ) )
-			continue;
-
-		NewFaceData.Add( FaceData[i] );
-	}
-
-	Swap( FaceData, NewFaceData );
-}
+}
\ No newline at end of file
diff --git a/Engine/Source/Developer/MeshSimplifier/Private/MeshTessellate.h b/Engine/Source/Developer/MeshSimplifier/Private/MeshTessellate.h
deleted file mode 100644
diff --git a/Engine/Source/Developer/MeshSimplifier/Private/QuadricMeshReduction.cpp b/Engine/Source/Developer/MeshSimplifier/Private/QuadricMeshReduction.cpp
--- a/Engine/Source/Developer/MeshSimplifier/Private/QuadricMeshReduction.cpp
+++ b/Engine/Source/Developer/MeshSimplifier/Private/QuadricMeshReduction.cpp
@@ -5,7 +5,7 @@
 #include "MeshUtilities.h"
 #include "MeshBuild.h"
 #include "MeshSimplify.h"
-#include "SimpVert.h"
+#include "LerpVert.h"
 #include "OverlappingCorners.h"
 #include "Templates/UniquePtr.h"
 #include "Features/IModularFeatures.h"
@@ -66,8 +66,7 @@ public:
 		// bool bUseQuadricSimplier = SplitVersionString[0].Equals("QuadricMeshReduction");
 
 		static FString Version = TEXT("QuadricMeshReduction_V2.1");
-		static FString AltVersion = TEXT("QuadricMeshReduction_V2.0_OldSimplifier");
-		return bUseOldMeshSimplifier ? AltVersion : Version;
+		return Version;
 	}
 
 	virtual void ReduceMeshDescription(
@@ -91,9 +90,9 @@ public:
 			FStaticMeshOperations::BuildWeldedVertexIDRemap(InMesh, ReductionSettings.WeldingThreshold, VertexIDRemap);
 		}
 
-		TArray< TVertSimp< NumTexCoords > >	Verts;
-		TArray< uint32 >					Indexes;
-		TArray< int32 >						MaterialIndexes;
+		TArray< FLerpVert >	Verts;
+		TArray< uint32 >	Indexes;
+		TArray< int32 >		MaterialIndexes;
 
 		TMap< int32, int32 > VertsMap;
 
@@ -142,38 +141,38 @@ public:
 				const int32 VertexInstanceValue = VertexInstanceID.GetValue();
 				const FVector3f& VertexPosition = CornerPositions[TriVert];
 
-				TVertSimp< NumTexCoords > NewVert;
+				FLerpVert NewVert;
 
 				NewVert.Position = CornerPositions[TriVert];
-				NewVert.Tangents[0] = InVertexTangents[ VertexInstanceID ];
-				NewVert.Normal = InVertexNormals[ VertexInstanceID ];
-				NewVert.Tangents[1] = FVector3f(0.0f);
-				if (!NewVert.Normal.IsNearlyZero(SMALL_NUMBER) && !NewVert.Tangents[0].IsNearlyZero(SMALL_NUMBER))
+				NewVert.TangentX = InVertexTangents[ VertexInstanceID ];
+				NewVert.TangentZ = InVertexNormals[ VertexInstanceID ];
+				NewVert.TangentY = FVector3f(0.0f);
+				if (!NewVert.TangentZ.IsNearlyZero(SMALL_NUMBER) && !NewVert.TangentX.IsNearlyZero(SMALL_NUMBER))
 				{
-					NewVert.Tangents[1] = FVector3f::CrossProduct(NewVert.Normal, NewVert.Tangents[0]).GetSafeNormal() * InVertexBinormalSigns[ VertexInstanceID ];
+					NewVert.TangentY = FVector3f::CrossProduct(NewVert.TangentZ, NewVert.TangentX).GetSafeNormal() * InVertexBinormalSigns[ VertexInstanceID ];
 				}
 
 				// Fix bad tangents
-				NewVert.Tangents[0] = NewVert.Tangents[0].ContainsNaN() ? FVector3f::ZeroVector : NewVert.Tangents[0];
-				NewVert.Tangents[1] = NewVert.Tangents[1].ContainsNaN() ? FVector3f::ZeroVector : NewVert.Tangents[1];
-				NewVert.Normal = NewVert.Normal.ContainsNaN() ? FVector3f::ZeroVector : NewVert.Normal;
+				NewVert.TangentX = NewVert.TangentX.ContainsNaN() ? FVector3f::ZeroVector : NewVert.TangentX;
+				NewVert.TangentY = NewVert.TangentY.ContainsNaN() ? FVector3f::ZeroVector : NewVert.TangentY;
+				NewVert.TangentZ = NewVert.TangentZ.ContainsNaN() ? FVector3f::ZeroVector : NewVert.TangentZ;
 				NewVert.Color = FLinearColor(InVertexColors[ VertexInstanceID ]);
 
 				for (int32 UVIndex = 0; UVIndex < NumTexCoords; UVIndex++)
 				{
 					if (UVIndex < InVertexUVs.GetNumChannels())
 					{
-						NewVert.TexCoords[UVIndex] = InVertexUVs.Get(VertexInstanceID, UVIndex);
+						NewVert.UVs[UVIndex] = InVertexUVs.Get(VertexInstanceID, UVIndex);
 						InMeshNumTexCoords = FMath::Max(UVIndex + 1, InMeshNumTexCoords);
 					}
 					else
 					{
-						NewVert.TexCoords[UVIndex] = FVector2f::ZeroVector;
+						NewVert.UVs[UVIndex] = FVector2f::ZeroVector;
 					}
 				}
 
 				// Make sure this vertex is valid from the start
-				NewVert.Correct();
+				CorrectAttributes( (float*)&NewVert.TangentX );
 					
 
 				//Never add duplicated vertex instance
@@ -192,13 +191,28 @@ public:
 					int32* Location = VertsMap.Find(DupVerts[k]);
 					if (Location)
 					{
-						TVertSimp< NumTexCoords >& FoundVert = Verts[*Location];
+						FLerpVert& FoundVert = Verts[*Location];
 
-						if (NewVert.Equals(FoundVert))
+						if( !PointsEqual(	NewVert.Position, FoundVert.Position ) ||
+							!NormalsEqual(	NewVert.TangentX, FoundVert.TangentX ) ||
+							!NormalsEqual(	NewVert.TangentY, FoundVert.TangentY ) ||
+							!NormalsEqual(	NewVert.TangentZ, FoundVert.TangentZ ) ||
+							!NewVert.Color.Equals( FoundVert.Color ) )
 						{
-							Index = *Location;
-							break;
+							continue;
 						}
+				
+						// UVs
+						for( int32 UVIndex = 0; UVIndex < NumTexCoords; UVIndex++ )
+						{
+							if( !UVsEqual( NewVert.UVs[ UVIndex ], FoundVert.UVs[ UVIndex ] ) )
+							{
+								continue;
+							}
+						}
+
+						Index = *Location;
+						break;
 					}
 				}
 				if (Index == INDEX_NONE)
@@ -259,183 +273,85 @@ public:
 		// Clamp to a minimum of 4 vertices, ReductionSettings.PercentVertices can be zero which makes TargetNumVerts also zero
 		TargetNumVerts = FMath::Max(TargetNumVerts, 4u);
 
-		if (bUseOldMeshSimplifier)
+		if( TargetNumVerts < NumVerts || TargetNumTris < NumTris )
 		{
-			static_assert(NumTexCoords == 8, "NumTexCoords changed, fix AttributeWeights");
-			const uint32 NumAttributes = (sizeof(TVertSimp< NumTexCoords >) - sizeof(FVector3f)) / sizeof(float);
-			float AttributeWeights[] =
+			const uint32 NumAttributes = ( sizeof( FLerpVert ) - sizeof( FVector3f ) ) / sizeof(float);
+			float AttributeWeights[ NumAttributes ] =
 			{
-				16.0f, 16.0f, 16.0f,	// Normal
-				0.1f, 0.1f, 0.1f,		// Tangent[0]
-				0.1f, 0.1f, 0.1f,		// Tangent[1]
-				0.1f, 0.1f, 0.1f, 0.1f,	// Color
-				0.5f, 0.5f,				// TexCoord[0]
-				0.5f, 0.5f,				// TexCoord[1]
-				0.5f, 0.5f,				// TexCoord[2]
-				0.5f, 0.5f,				// TexCoord[3]
-				0.5f, 0.5f,				// TexCoord[4]
-				0.5f, 0.5f,				// TexCoord[5]
-				0.5f, 0.5f,				// TexCoord[6]
-				0.5f, 0.5f,				// TexCoord[7]
+				16.0f, 16.0f, 16.0f,// Normal
+				0.1f, 0.1f, 0.1f,	// Tangent[0]
+				0.1f, 0.1f, 0.1f	// Tangent[1]
 			};
-			float* ColorWeights = AttributeWeights + 3 + 3 + 3;
-			float* TexCoordWeights = ColorWeights + 4;
-
-			// Re-scale the weights for UV channels that exceed the expected 0-1 range.
-			// Otherwise garbage on the UVs will dominate the simplification quadric.
-			{
-				float XLength[MAX_STATIC_TEXCOORDS] = { 0 };
-				float YLength[MAX_STATIC_TEXCOORDS] = { 0 };
-				{
-					for (int32 TexCoordId = 0; TexCoordId < NumTexCoords; ++TexCoordId)
-					{
-						float XMax = -FLT_MAX;
-						float YMax = -FLT_MAX;
-						float XMin = FLT_MAX;
-						float YMin = FLT_MAX;
-						for (const TVertSimp< NumTexCoords >& SimpVert : Verts)
-						{
-							const FVector2f& UVs = SimpVert.TexCoords[TexCoordId];
-							XMax = FMath::Max(XMax, UVs.X);
-							XMin = FMath::Min(XMin, UVs.X);
-
-							YMax = FMath::Max(YMax, UVs.Y);
-							YMin = FMath::Min(YMin, UVs.Y);
-						}
-
-						XLength[TexCoordId] = (XMax > XMin) ? XMax - XMin : 0.f;
-						YLength[TexCoordId] = (YMax > YMin) ? YMax - YMin : 0.f;
-					}
-				}
-
-				for (int32 TexCoordId = 0; TexCoordId < NumTexCoords; ++TexCoordId)
-				{
+			float* ColorWeights = AttributeWeights + 9;
+			float* UVWeights = ColorWeights + 4;
 
-					if (XLength[TexCoordId] > 1.f)
-					{
-						TexCoordWeights[2 * TexCoordId + 0] /= XLength[TexCoordId];
-					}
-					if (YLength[TexCoordId] > 1.f)
-					{
-						TexCoordWeights[2 * TexCoordId + 1] /= YLength[TexCoordId];
-					}
-				}
-			}
+			bool bHasColors = true;
 
-			// Zero out weights that aren't used
+			// Set weights if they are used
+			if( bHasColors )
 			{
-				//TODO Check if we have vertex color
-
-				for (int32 TexCoordIndex = 0; TexCoordIndex < NumTexCoords; TexCoordIndex++)
-				{
-					if (TexCoordIndex >= InVertexUVs.GetNumChannels())
-					{
-						TexCoordWeights[2 * TexCoordIndex + 0] = 0.0f;
-						TexCoordWeights[2 * TexCoordIndex + 1] = 0.0f;
-					}
-				}
+				ColorWeights[0] = 0.1f;
+				ColorWeights[1] = 0.1f;
+				ColorWeights[2] = 0.1f;
+				ColorWeights[3] = 0.1f;
 			}
 
-			TMeshSimplifier< TVertSimp< NumTexCoords >, NumAttributes >* MeshSimp = new TMeshSimplifier< TVertSimp< NumTexCoords >, NumAttributes >(Verts.GetData(), NumVerts, Indexes.GetData(), NumIndexes);
-
-			MeshSimp->SetAttributeWeights(AttributeWeights);
-			MeshSimp->SetEdgeWeight(256.0f);
-			//MeshSimp->SetBoundaryLocked();
-			MeshSimp->InitCosts();
-
-			float MaxErrorSqr = MeshSimp->SimplifyMesh(MAX_FLT, TargetNumTris, TargetNumVerts);
-
-			MeshSimp->OutputMesh(Verts.GetData(), Indexes.GetData(), &NumVerts, &NumIndexes);
-			MeshSimp->CompactFaceData(MaterialIndexes);
-			NumTris = NumIndexes / 3;
-			delete MeshSimp;
-
-			OutMaxDeviation = FMath::Sqrt(MaxErrorSqr) / 8.0f;
-		}
-		else
-		{
-			if( TargetNumVerts < NumVerts || TargetNumTris < NumTris )
+			float UVWeight = 0.5f;
+			for( int32 UVIndex = 0; UVIndex < InVertexUVs.GetNumChannels(); UVIndex++ )
 			{
-				using VertType = TVertSimp< NumTexCoords >;
-
-				const uint32 NumAttributes = ( sizeof( VertType ) - sizeof( FVector3f ) ) / sizeof(float);
-				float AttributeWeights[ NumAttributes ] =
-				{
-					16.0f, 16.0f, 16.0f,// Normal
-					0.1f, 0.1f, 0.1f,	// Tangent[0]
-					0.1f, 0.1f, 0.1f	// Tangent[1]
-				};
-				float* ColorWeights = AttributeWeights + 9;
-				float* UVWeights = ColorWeights + 4;
-
-				bool bHasColors = true;
+				// Normalize UVWeights using min/max UV range.
 
-				// Set weights if they are used
-				if( bHasColors )
-				{
-					ColorWeights[0] = 0.1f;
-					ColorWeights[1] = 0.1f;
-					ColorWeights[2] = 0.1f;
-					ColorWeights[3] = 0.1f;
-				}
+				float MinUV = +FLT_MAX;
+				float MaxUV = -FLT_MAX;
 
-				float UVWeight = 0.5f;
-				for( int32 UVIndex = 0; UVIndex < InVertexUVs.GetNumChannels(); UVIndex++ )
+				for( int32 VertexIndex = 0; VertexIndex < Verts.Num(); VertexIndex++ )
 				{
-					// Normalize UVWeights using min/max UV range.
-
-					float MinUV = +FLT_MAX;
-					float MaxUV = -FLT_MAX;
-
-					for( int32 VertexIndex = 0; VertexIndex < Verts.Num(); VertexIndex++ )
-					{
-						MinUV = FMath::Min( MinUV, Verts[ VertexIndex ].TexCoords[ UVIndex ].X );
-						MinUV = FMath::Min( MinUV, Verts[ VertexIndex ].TexCoords[ UVIndex ].Y );
-						MaxUV = FMath::Max( MaxUV, Verts[ VertexIndex ].TexCoords[ UVIndex ].X );
-						MaxUV = FMath::Max( MaxUV, Verts[ VertexIndex ].TexCoords[ UVIndex ].Y );
-					}
-
-					UVWeights[ 2 * UVIndex + 0 ] = UVWeight / FMath::Max( 1.0f, MaxUV - MinUV );
-					UVWeights[ 2 * UVIndex + 1 ] = UVWeight / FMath::Max( 1.0f, MaxUV - MinUV );
+					MinUV = FMath::Min( MinUV, Verts[ VertexIndex ].UVs[ UVIndex ].X );
+					MinUV = FMath::Min( MinUV, Verts[ VertexIndex ].UVs[ UVIndex ].Y );
+					MaxUV = FMath::Max( MaxUV, Verts[ VertexIndex ].UVs[ UVIndex ].X );
+					MaxUV = FMath::Max( MaxUV, Verts[ VertexIndex ].UVs[ UVIndex ].Y );
 				}
 
-				FMeshSimplifier Simplifier( (float*)Verts.GetData(), Verts.Num(), Indexes.GetData(), Indexes.Num(), MaterialIndexes.GetData(), NumAttributes );
-
-				Simplifier.SetAttributeWeights( AttributeWeights );
-				Simplifier.SetCorrectAttributes( CorrectAttributes );
-				Simplifier.SetEdgeWeight( 512.0f );
-				Simplifier.SetLimitErrorToSurfaceArea( false );
-
-				Simplifier.DegreePenalty = 100.0f;
-				Simplifier.InversionPenalty = 1000000.0f;
+				UVWeights[ 2 * UVIndex + 0 ] = UVWeight / FMath::Max( 1.0f, MaxUV - MinUV );
+				UVWeights[ 2 * UVIndex + 1 ] = UVWeight / FMath::Max( 1.0f, MaxUV - MinUV );
+			}
 
-				float MaxErrorSqr = Simplifier.Simplify( TargetNumVerts, TargetNumTris, 0.0f, 4, 2, MAX_flt );
+			FMeshSimplifier Simplifier( (float*)Verts.GetData(), Verts.Num(), Indexes.GetData(), Indexes.Num(), MaterialIndexes.GetData(), NumAttributes );
 
-				if( Simplifier.GetRemainingNumVerts() == 0 || Simplifier.GetRemainingNumTris() == 0 )
-				{
-					// Reduced to nothing so just return the orignial.
-					OutReducedMesh = InMesh;
-					OutMaxDeviation = 0.0f;
-					return;
-				}
-		
-				Simplifier.Compact();
+			Simplifier.SetAttributeWeights( AttributeWeights );
+			Simplifier.SetCorrectAttributes( CorrectAttributes );
+			Simplifier.SetEdgeWeight( 512.0f );
+			Simplifier.SetLimitErrorToSurfaceArea( false );
 
-				Verts.SetNum( Simplifier.GetRemainingNumVerts() );
-				Indexes.SetNum( Simplifier.GetRemainingNumTris() * 3 );
-				MaterialIndexes.SetNum( Simplifier.GetRemainingNumTris() );
+			Simplifier.DegreePenalty = 100.0f;
+			Simplifier.InversionPenalty = 1000000.0f;
 
-				NumVerts = Simplifier.GetRemainingNumVerts();
-				NumTris = Simplifier.GetRemainingNumTris();
-				NumIndexes = NumTris * 3;
+			float MaxErrorSqr = Simplifier.Simplify( TargetNumVerts, TargetNumTris, 0.0f, 4, 2, MAX_flt );
 
-				OutMaxDeviation = FMath::Sqrt( MaxErrorSqr ) / 8.0f;
-			}
-			else
+			if( Simplifier.GetRemainingNumVerts() == 0 || Simplifier.GetRemainingNumTris() == 0 )
 			{
-				// Rare but could happen with rounding or only 2 triangles.
+				// Reduced to nothing so just return the orignial.
+				OutReducedMesh = InMesh;
 				OutMaxDeviation = 0.0f;
+				return;
 			}
+	
+			Simplifier.Compact();
+
+			Verts.SetNum( Simplifier.GetRemainingNumVerts() );
+			Indexes.SetNum( Simplifier.GetRemainingNumTris() * 3 );
+			MaterialIndexes.SetNum( Simplifier.GetRemainingNumTris() );
+
+			NumVerts = Simplifier.GetRemainingNumVerts();
+			NumTris = Simplifier.GetRemainingNumTris();
+			NumIndexes = NumTris * 3;
+
+			OutMaxDeviation = FMath::Sqrt( MaxErrorSqr ) / 8.0f;
+		}
+		else
+		{
+			// Rare but could happen with rounding or only 2 triangles.
+			OutMaxDeviation = 0.0f;
 		}
 
 		{
@@ -492,9 +408,12 @@ public:
 					check(AddedVertexInstanceId.GetValue() == VertexInstanceIndex);
 
 					//NTBs information
-					OutVertexTangents[AddedVertexInstanceId] = Verts[Indexes[VertexInstanceIndex]].Tangents[0];
-					OutVertexBinormalSigns[AddedVertexInstanceId] = GetBasisDeterminantSign((FVector)Verts[Indexes[VertexInstanceIndex]].Tangents[0].GetSafeNormal(), (FVector)Verts[Indexes[VertexInstanceIndex]].Tangents[1].GetSafeNormal(), (FVector)Verts[Indexes[VertexInstanceIndex]].Normal.GetSafeNormal());
-					OutVertexNormals[AddedVertexInstanceId] = Verts[Indexes[VertexInstanceIndex]].Normal;
+					OutVertexTangents[AddedVertexInstanceId] = Verts[Indexes[VertexInstanceIndex]].TangentX;
+					OutVertexBinormalSigns[AddedVertexInstanceId] = GetBasisDeterminantSign(
+						(FVector)Verts[Indexes[VertexInstanceIndex]].TangentX,
+						(FVector)Verts[Indexes[VertexInstanceIndex]].TangentY,
+						(FVector)Verts[Indexes[VertexInstanceIndex]].TangentZ);
+					OutVertexNormals[AddedVertexInstanceId] = Verts[Indexes[VertexInstanceIndex]].TangentZ;
 
 					//Vertex Color
 					OutVertexColors[AddedVertexInstanceId] = Verts[Indexes[VertexInstanceIndex]].Color;
@@ -502,7 +421,7 @@ public:
 					//Texture coord
 					for (int32 TexCoordIndex = 0; TexCoordIndex < InMeshNumTexCoords; TexCoordIndex++)
 					{
-						OutVertexUVs.Set(AddedVertexInstanceId, TexCoordIndex, Verts[Indexes[VertexInstanceIndex]].TexCoords[TexCoordIndex]);
+						OutVertexUVs.Set(AddedVertexInstanceId, TexCoordIndex, Verts[Indexes[VertexInstanceIndex]].UVs[TexCoordIndex]);
 					}
 				}
 				
@@ -624,20 +543,12 @@ public:
 		return false;
 	}
 
-	FQuadricSimplifierMeshReduction()
-		: bUseOldMeshSimplifier(false)
-	{
-		GConfig->GetBool(TEXT("/Script/Engine.MeshSimplificationSettings"), TEXT("bMeshReductionBackwardCompatible"), bUseOldMeshSimplifier, GEngineIni);
-	}
-
 	virtual ~FQuadricSimplifierMeshReduction() {}
 
 	static FQuadricSimplifierMeshReduction* Create()
 	{
 		return new FQuadricSimplifierMeshReduction;
 	}
-
-	bool bUseOldMeshSimplifier;
 };
 
 TUniquePtr<FQuadricSimplifierMeshReduction> GQuadricSimplifierMeshReduction;
diff --git a/Engine/Source/Developer/MeshSimplifier/Private/SimpVert.h b/Engine/Source/Developer/MeshSimplifier/Private/SimpVert.h
deleted file mode 100644
diff --git a/Engine/Source/Developer/MeshSimplifier/QuadricMeshReduction.Build.cs b/Engine/Source/Developer/MeshSimplifier/QuadricMeshReduction.Build.cs
--- a/Engine/Source/Developer/MeshSimplifier/QuadricMeshReduction.Build.cs
+++ b/Engine/Source/Developer/MeshSimplifier/QuadricMeshReduction.Build.cs
@@ -7,13 +7,19 @@ public class QuadricMeshReduction : ModuleRules
 {
 	public QuadricMeshReduction(ReadOnlyTargetRules Target) : base(Target)
 	{
-		PrivateDependencyModuleNames.Add("Core");
-		PrivateDependencyModuleNames.Add("CoreUObject");
-		PrivateDependencyModuleNames.Add("Engine");
-		PrivateDependencyModuleNames.Add("RenderCore");
-        PrivateDependencyModuleNames.Add("MeshDescription");
-        PrivateDependencyModuleNames.Add("MeshUtilitiesCommon");
-        PrivateDependencyModuleNames.Add("StaticMeshDescription");
+		PrivateDependencyModuleNames.AddRange(
+			new string[]
+			{
+				"Core",
+				"CoreUObject",
+				"Engine",
+				"RenderCore",
+				"NaniteUtilities",
+				"MeshUtilitiesCommon",
+				"MeshDescription",
+				"StaticMeshDescription",
+			}
+		);
 
         PrivateIncludePathModuleNames.AddRange(
         new string[] {
