Fix mesh reduction bugs due to replacing TSimpVert with FLerpVert.

diff --git a/Engine/Source/Developer/MeshSimplifier/Private/QuadricMeshReduction.cpp b/Engine/Source/Developer/MeshSimplifier/Private/QuadricMeshReduction.cpp
--- a/Engine/Source/Developer/MeshSimplifier/Private/QuadricMeshReduction.cpp
+++ b/Engine/Source/Developer/MeshSimplifier/Private/QuadricMeshReduction.cpp
@@ -40,20 +40,43 @@ IMPLEMENT_MODULE(FQuadricSimplifierMeshReductionModule, QuadricMeshReduction);
 
 void CorrectAttributes( float* Attributes )
 {
-	FVector3f& Normal	= *reinterpret_cast< FVector3f* >( Attributes );
-	FVector3f& TangentX	= *reinterpret_cast< FVector3f* >( Attributes + 3 );
-	FVector3f& TangentY	= *reinterpret_cast< FVector3f* >( Attributes + 3 + 3 );
+	FVector3f& TangentX	= *reinterpret_cast< FVector3f* >( Attributes );
+	FVector3f& TangentY	= *reinterpret_cast< FVector3f* >( Attributes + 3 );
+	FVector3f& TangentZ	= *reinterpret_cast< FVector3f* >( Attributes + 3 + 3 );
 	FLinearColor& Color	= *reinterpret_cast< FLinearColor* >( Attributes + 3 + 3 + 3 );
 
-	Normal.Normalize();
-	TangentX -= ( TangentX | Normal ) * Normal;
+	TangentZ.Normalize();
+	TangentX -= ( TangentX | TangentZ ) * TangentZ;
 	TangentX.Normalize();
-	TangentY -= ( TangentY | Normal ) * Normal;
+	TangentY -= ( TangentY | TangentZ ) * TangentZ;
 	TangentY -= ( TangentY | TangentX ) * TangentX;
 	TangentY.Normalize();
 	Color = Color.GetClamped();
 }
 
+bool VertsEqual( const FLerpVert& A, const FLerpVert& B )
+{
+	if( !PointsEqual(	A.Position, B.Position ) ||
+		!NormalsEqual(	A.TangentX, B.TangentX ) ||
+		!NormalsEqual(	A.TangentY, B.TangentY ) ||
+		!NormalsEqual(	A.TangentZ, B.TangentZ ) ||
+		!A.Color.Equals( B.Color ) )
+	{
+		return false;
+	}
+
+	// UVs
+	for( int32 UVIndex = 0; UVIndex < MAX_STATIC_TEXCOORDS; UVIndex++ )
+	{
+		if( !UVsEqual( A.UVs[ UVIndex ], B.UVs[ UVIndex ] ) )
+		{
+			return false;
+		}
+	}
+
+	return true;
+}
+
 class FQuadricSimplifierMeshReduction : public IMeshReduction
 {
 public:
@@ -192,27 +215,11 @@ public:
 					if (Location)
 					{
 						FLerpVert& FoundVert = Verts[*Location];
-
-						if( !PointsEqual(	NewVert.Position, FoundVert.Position ) ||
-							!NormalsEqual(	NewVert.TangentX, FoundVert.TangentX ) ||
-							!NormalsEqual(	NewVert.TangentY, FoundVert.TangentY ) ||
-							!NormalsEqual(	NewVert.TangentZ, FoundVert.TangentZ ) ||
-							!NewVert.Color.Equals( FoundVert.Color ) )
-						{
-							continue;
-						}
-				
-						// UVs
-						for( int32 UVIndex = 0; UVIndex < NumTexCoords; UVIndex++ )
+						if( VertsEqual( NewVert, FoundVert ) )
 						{
-							if( !UVsEqual( NewVert.UVs[ UVIndex ], FoundVert.UVs[ UVIndex ] ) )
-							{
-								continue;
-							}
+							Index = *Location;
+							break;
 						}
-
-						Index = *Location;
-						break;
 					}
 				}
 				if (Index == INDEX_NONE)
@@ -278,9 +285,9 @@ public:
 			const uint32 NumAttributes = ( sizeof( FLerpVert ) - sizeof( FVector3f ) ) / sizeof(float);
 			float AttributeWeights[ NumAttributes ] =
 			{
-				16.0f, 16.0f, 16.0f,// Normal
 				0.1f, 0.1f, 0.1f,	// Tangent[0]
-				0.1f, 0.1f, 0.1f	// Tangent[1]
+				0.1f, 0.1f, 0.1f,	// Tangent[1]
+				16.0f, 16.0f, 16.0f	// Normal
 			};
 			float* ColorWeights = AttributeWeights + 9;
 			float* UVWeights = ColorWeights + 4;
diff --git a/Engine/Source/Runtime/Core/Private/UObject/UE5ReleaseStreamObjectVersion.cpp b/Engine/Source/Runtime/Core/Private/UObject/UE5ReleaseStreamObjectVersion.cpp
--- a/Engine/Source/Runtime/Core/Private/UObject/UE5ReleaseStreamObjectVersion.cpp
+++ b/Engine/Source/Runtime/Core/Private/UObject/UE5ReleaseStreamObjectVersion.cpp
@@ -14,6 +14,7 @@ TMap<FGuid, FGuid> FUE5ReleaseStreamObjectVersion::GetSystemGuids()
 	SystemGuids.Add(DevGuids.NIAGARASHADERMAP_DERIVEDDATA_VER, FGuid("6BE35B6FACB34568970120F9BC8DAB80"));
 	SystemGuids.Add(DevGuids.Niagara_LatestScriptCompileVersion, FGuid("43A777C54EB24E9894525B8D04529F23"));
 	SystemGuids.Add(DevGuids.SkeletalMeshDerivedDataVersion, FGuid("25C49E579B3142DDA2A8C14037267679"));
+	SystemGuids.Add(DevGuids.STATICMESH_DERIVEDDATA_VER, FGuid("7A556175518D458AA8786ED69AB8DDE7"));
 
 	return SystemGuids;
 }
