Add a check to the mesh simplifier to avoid collapsing an edge if it is shared by all triangles, thus causing all triangles in the cluster to be degenerate.
This fixes asserts that have been seen in the wild.

diff --git a/Engine/Source/Developer/MeshSimplifier/Private/MeshSimplify.cpp b/Engine/Source/Developer/MeshSimplifier/Private/MeshSimplify.cpp
--- a/Engine/Source/Developer/MeshSimplifier/Private/MeshSimplify.cpp
+++ b/Engine/Source/Developer/MeshSimplifier/Private/MeshSimplify.cpp
@@ -262,9 +262,11 @@ float FMeshSimplifier::EvaluateMerge( const FVector3f& Position0, const FVector3
 	GatherAdjTris( Position1, 1, FlagsUnion1 );
 
 	if( VertDegree == 0 )
-	{
 		return 0.0f;
-	}
+
+	// This would mean this collapse will remove all remaining triangles.
+	if( VertDegree == RemainingNumTris * 2 )
+		return 0.0f;
 
 	bool bLocked0 = FlagsUnion0 & LockedVertMask;
 	bool bLocked1 = FlagsUnion1 & LockedVertMask;
