Add a check to the mesh simplifier to avoid collapsing an edge if it is shared by all triangles, thus causing all triangles in the cluster to be degenerate
This fixes asserts that have been seen in the wild.

NOTE: This is a resubmission, with an additional fix to an assert the original change inadvertently caused.

diff --git a/Engine/Source/Developer/MeshSimplifier/Private/MeshSimplify.cpp b/Engine/Source/Developer/MeshSimplifier/Private/MeshSimplify.cpp
--- a/Engine/Source/Developer/MeshSimplifier/Private/MeshSimplify.cpp
+++ b/Engine/Source/Developer/MeshSimplifier/Private/MeshSimplify.cpp
@@ -266,6 +266,21 @@ float FMeshSimplifier::EvaluateMerge( const FVector3f& Position0, const FVector3
 		return 0.0f;
 	}
 
+	// This would mean this collapse will remove all remaining triangles.
+	if( VertDegree == RemainingNumTris * 2 )
+	{
+		// Clean up corner flags
+		check( AdjTris.Num() > 0 );
+		for( uint32 TriIndex : AdjTris )
+		{
+			for( uint32 CornerIndex = 0; CornerIndex < 3; CornerIndex++ )
+			{
+				CornerFlags[ TriIndex * 3 + CornerIndex ] &= ~( MergeMask | AdjTriMask );
+			}
+		}
+		return 0.0f;
+	}
+
 	bool bLocked0 = FlagsUnion0 & LockedVertMask;
 	bool bLocked1 = FlagsUnion1 & LockedVertMask;
 
