Remove unnecessary include

diff --git a/Engine/Source/Developer/MeshSimplifier/Private/QuadricMeshReduction.cpp b/Engine/Source/Developer/MeshSimplifier/Private/QuadricMeshReduction.cpp
--- a/Engine/Source/Developer/MeshSimplifier/Private/QuadricMeshReduction.cpp
+++ b/Engine/Source/Developer/MeshSimplifier/Private/QuadricMeshReduction.cpp
@@ -2,7 +2,6 @@
 
 #include "CoreMinimal.h"
 #include "Modules/ModuleManager.h"
-#include "MeshUtilities.h"
 #include "MeshBuild.h"
 #include "MeshSimplify.h"
 #include "LerpVert.h"
