Fix regression of LOD generation (when Nanite is enabled) that would allow meshes with with low poly counts to simplify too much in generated LODs and destroy their silhouettes.

diff --git a/Engine/Source/Developer/MeshSimplifier/Private/QuadricMeshReduction.cpp b/Engine/Source/Developer/MeshSimplifier/Private/QuadricMeshReduction.cpp
--- a/Engine/Source/Developer/MeshSimplifier/Private/QuadricMeshReduction.cpp
+++ b/Engine/Source/Developer/MeshSimplifier/Private/QuadricMeshReduction.cpp
@@ -274,8 +274,8 @@ public:
 			TargetNumVerts = FMath::Min(ReductionSettings.MaxNumOfVerts, TargetNumVerts);
 		}
 
-		// We need a minimum of 2 triangles, to see the object on both side. If we use one, we will end up with zero triangle when we will remove a shared edge
-		TargetNumTris = FMath::Max(TargetNumTris, 2u);
+		// Keep a minimum of 64 triangles, to prevent already low-poly meshes from losing too much of their silhouette.
+		TargetNumTris = FMath::Max(TargetNumTris, 64u);
 		// Clamp to a minimum of 4 vertices, ReductionSettings.PercentVertices can be zero which makes TargetNumVerts also zero
 		TargetNumVerts = FMath::Max(TargetNumVerts, 4u);
 
diff --git a/Engine/Source/Runtime/Core/Private/UObject/FortniteSeasonBranchObjectVersion.cpp b/Engine/Source/Runtime/Core/Private/UObject/FortniteSeasonBranchObjectVersion.cpp
--- a/Engine/Source/Runtime/Core/Private/UObject/FortniteSeasonBranchObjectVersion.cpp
+++ b/Engine/Source/Runtime/Core/Private/UObject/FortniteSeasonBranchObjectVersion.cpp
@@ -6,7 +6,7 @@ TMap<FGuid, FGuid> FFortniteSeasonBranchObjectVersion::GetSystemGuids()
 	TMap<FGuid, FGuid> SystemGuids;
 	const FDevSystemGuids& DevGuids = FDevSystemGuids::Get();
 
-	SystemGuids.Add(DevGuids.STATICMESH_DERIVEDDATA_VER, FGuid("F16D2728168B4CDEBFB382B38FFE9E06"));
+	SystemGuids.Add(DevGuids.STATICMESH_DERIVEDDATA_VER, FGuid("C8B712A947FA4A42B783183B1176DB99"));
 	SystemGuids.Add(DevGuids.NANITE_DERIVEDDATA_VER, FGuid("D8BF6DDDAAD748BC9854D5F56631973C"));
 	SystemGuids.Add(DevGuids.Niagara_LatestScriptCompileVersion, FGuid("08A58144744FCA4A9BF6DD45DFC4EC13"));
 
